Add low-level unit test project for ToolMenus

Includes a few tests to verify that the tests themselves work.

diff --git a/Engine/Source/Programs/ToolMenusTests/Tests/TestGroupEvents.cpp b/Engine/Source/Programs/ToolMenusTests/Tests/TestGroupEvents.cpp
new file mode 100644
--- /dev/null
+++ b/Engine/Source/Programs/ToolMenusTests/Tests/TestGroupEvents.cpp
@@ -0,0 +1,23 @@
+// Copyright Epic Games, Inc. All Rights Reserved.
+#include "TestCommon/CoreUtilities.h"
+#include "TestCommon/CoreUObjectUtilities.h"
+
+#include <catch2/catch_test_macros.hpp>
+
+GROUP_BEFORE_GLOBAL(Catch::DefaultGroup)
+{
+	InitTaskGraph();
+	InitThreadPool(true);
+#if WITH_COREUOBJECT
+	InitCoreUObject();
+#endif
+}
+
+GROUP_AFTER_GLOBAL(Catch::DefaultGroup)
+{
+#if WITH_COREUOBJECT
+	CleanupCoreUObject();
+#endif
+	CleanupThreadPool();
+	CleanupTaskGraph();
+}
diff --git a/Engine/Source/Programs/ToolMenusTests/Tests/ToolMenus/ToolMenusTest.cpp b/Engine/Source/Programs/ToolMenusTests/Tests/ToolMenus/ToolMenusTest.cpp
new file mode 100644
--- /dev/null
+++ b/Engine/Source/Programs/ToolMenusTests/Tests/ToolMenus/ToolMenusTest.cpp
@@ -0,0 +1,60 @@
+// Copyright Epic Games, Inc. All Rights Reserved.
+
+#include "ToolMenus.h"
+
+#include "TestHarness.h"
+
+UToolMenus* CreateUniqueUToolMenusInstance()
+{
+	// Create a unique ToolMenus instance for use in a single test.
+	UToolMenus* ToolMenus = NewObject<UToolMenus>();
+	ToolMenus->AddToRoot();
+	return ToolMenus;
+}
+
+TEST_CASE("Developer::ToolMenus::Can create menu", "[ToolMenus]")
+{
+	UToolMenus* ToolMenus = CreateUniqueUToolMenusInstance();
+
+	UToolMenu* ToolMenu = ToolMenus->RegisterMenu("MyMenu");
+
+	CHECK(ToolMenu);
+}
+
+TEST_CASE("Developer::ToolMenus::Non-registered menu is not registered", "[ToolMenus]")
+{
+	UToolMenus* ToolMenus = CreateUniqueUToolMenusInstance();
+
+	CHECK_FALSE(ToolMenus->IsMenuRegistered("MyMenu"));
+}
+
+TEST_CASE("Developer::ToolMenus::Removed menu is not registered", "[ToolMenus]")
+{
+	UToolMenus* ToolMenus = CreateUniqueUToolMenusInstance();
+
+	ToolMenus->RegisterMenu("MyMenu");
+
+	CHECK(ToolMenus->IsMenuRegistered("MyMenu"));
+
+	ToolMenus->RemoveMenu("MyMenu");
+
+	CHECK_FALSE(ToolMenus->IsMenuRegistered("MyMenu"));
+}
+
+TEST_CASE("Developer::ToolMenus::GenerateWidget calls dynamic section lambdas", "[ToolMenus]")
+{
+	UToolMenus* ToolMenus = CreateUniqueUToolMenusInstance();
+
+	UToolMenu* ToolMenu = ToolMenus->RegisterMenu("MyMenu");
+
+	REQUIRE(ToolMenu);
+
+	bool bWasLambdaCalled = false;
+	ToolMenu->AddDynamicSection("MySection", FNewToolMenuDelegateLegacy::CreateLambda([&bWasLambdaCalled](FMenuBuilder&, UToolMenu*) {
+		bWasLambdaCalled = true;
+	}));
+
+	ToolMenus->GenerateWidget("MyMenu", FToolMenuContext());
+
+	CHECK(bWasLambdaCalled);
+}
\ No newline at end of file
diff --git a/Engine/Source/Programs/ToolMenusTests/ToolMenusTests.Build.cs b/Engine/Source/Programs/ToolMenusTests/ToolMenusTests.Build.cs
new file mode 100644
--- /dev/null
+++ b/Engine/Source/Programs/ToolMenusTests/ToolMenusTests.Build.cs
@@ -0,0 +1,57 @@
+// Copyright Epic Games, Inc. All Rights Reserved.
+using UnrealBuildTool;
+
+public class ToolMenusTests : TestModuleRules
+{
+	protected Metadata ToolMenusTestsMetadata = new Metadata() {
+		TestName = "ToolMenus",
+		TestShortName = "ToolMenus",
+		ReportType = "xml",
+		SupportedPlatforms = {
+			UnrealTargetPlatform.Win64,
+			UnrealTargetPlatform.Linux,
+			UnrealTargetPlatform.Mac } };
+
+	/// <summary>
+	/// Test metadata to be used with BuildGraph
+	/// </summary>
+	public Metadata TestMetadata
+	{ 
+		get { return ToolMenusTestsMetadata; }
+	}
+
+	public ToolMenusTests(ReadOnlyTargetRules Target) : base(Target, true)
+	{
+		PrivateDependencyModuleNames.AddRange(
+			new string[] {
+				"Core",
+				"CoreUObject",
+				"Slate",
+				"SlateCore",
+				"ToolMenus"
+			});
+
+		if (Target.bBuildWithEditorOnlyData)
+		{
+			PrivateDependencyModuleNames.AddRange(
+				new string[] {
+					"DesktopPlatform"
+				});
+		}
+
+		string PlatformCompilationArgs;
+		foreach (var Platform in UnrealTargetPlatform.GetValidPlatforms())
+		{
+			if (Platform == UnrealTargetPlatform.Android)
+			{
+				PlatformCompilationArgs = "-allmodules -architectures=arm64";
+			}
+			else
+			{
+				PlatformCompilationArgs = "-allmodules";
+			}
+			TestMetadata.PlatformCompilationExtraArgs.Add(Platform, PlatformCompilationArgs);
+		}
+		UpdateBuildGraphPropertiesFile(TestMetadata);
+	}
+}
diff --git a/Engine/Source/Programs/ToolMenusTests/ToolMenusTests.Target.cs b/Engine/Source/Programs/ToolMenusTests/ToolMenusTests.Target.cs
new file mode 100644
--- /dev/null
+++ b/Engine/Source/Programs/ToolMenusTests/ToolMenusTests.Target.cs
@@ -0,0 +1,16 @@
+// Copyright Epic Games, Inc. All Rights Reserved.
+
+using UnrealBuildTool;
+
+[SupportedPlatforms(UnrealPlatformClass.All)]
+public class ToolMenusTestsTarget : TestTargetRules
+{
+	public ToolMenusTestsTarget(TargetInfo Target) : base(Target)
+	{
+		bTestsRequireApplicationCore = true;
+		bTestsRequireCoreUObject = true;
+
+		bBuildWithEditorOnlyData = Target.Platform.IsInGroup(UnrealPlatformGroup.Desktop)
+			&& (Target.Configuration == UnrealTargetConfiguration.Debug || Target.Configuration == UnrealTargetConfiguration.Development);
+	}
+}
