ToolMenus: Add a method to unregister an owner for menu profiles
ToolMenusTest: Add tests for sections, entries, customizations and profiles

diff --git a/Engine/Source/Developer/ToolMenus/Private/ToolMenus.cpp b/Engine/Source/Developer/ToolMenus/Private/ToolMenus.cpp
--- a/Engine/Source/Developer/ToolMenus/Private/ToolMenus.cpp
+++ b/Engine/Source/Developer/ToolMenus/Private/ToolMenus.cpp
@@ -2538,7 +2538,7 @@ void UToolMenus::UnregisterOwnerInternal(FToolMenuOwner InOwner)
 	}
 }
 
-void UToolMenus::UnregisterRuntimeMenuCustomizationOwner(FName InOwnerName)
+void UToolMenus::UnregisterRuntimeMenuCustomizationOwner(const FName InOwnerName)
 {
 	if (InOwnerName.IsNone())
 	{
@@ -2566,6 +2566,41 @@ void UToolMenus::UnregisterRuntimeMenuCustomizationOwner(FName InOwnerName)
 	}
 }
 
+void UToolMenus::UnregisterRuntimeMenuProfileOwner(const FName InOwnerName)
+{
+	if (InOwnerName.IsNone())
+	{
+		return;
+	}
+
+	bool bNeedsRefresh = false;
+
+	// Loop through all menus with profiles
+	for (TPair<FName, FToolMenuProfileMap>& MenusWithProfiles : RuntimeMenuProfiles)
+	{
+		// Loop through all profiles for a given menu
+		for (TPair<FName, FToolMenuProfile>& MenuProfile : MenusWithProfiles.Value.MenuProfiles)
+		{
+			if (MenuProfile.Value.MenuPermissions.UnregisterOwner(InOwnerName))
+			{
+				bNeedsRefresh = true;
+			}
+
+			if (MenuProfile.Value.SuppressExtenders.Remove(InOwnerName) > 0)
+			{
+				bNeedsRefresh = true;
+			}
+		}
+	}
+
+	// Refresh any widgets that are currently displayed to the user
+	if (bNeedsRefresh)
+	{
+		RefreshAllWidgets();
+	}
+}
+
+
 FToolMenuOwner UToolMenus::CurrentOwner() const
 {
 	if (OwnerStack.Num() > 0)
diff --git a/Engine/Source/Developer/ToolMenus/Public/ToolMenus.h b/Engine/Source/Developer/ToolMenus/Public/ToolMenus.h
--- a/Engine/Source/Developer/ToolMenus/Public/ToolMenus.h
+++ b/Engine/Source/Developer/ToolMenus/Public/ToolMenus.h
@@ -293,6 +293,9 @@ public:
 	/** Unregister runtime customization settings for a specific owner name */
 	void UnregisterRuntimeMenuCustomizationOwner(const FName InOwnerName);
 
+	/** Unregister runtime profile settings for a specific owner name */
+	void UnregisterRuntimeMenuProfileOwner(const FName InOwnerName);
+
 	/** Generates sub menu by entry name in the given generated menu parent */
 	UToolMenu* GenerateSubMenu(const UToolMenu* InGeneratedParent, const FName InBlockName);
 
diff --git a/Engine/Source/Programs/ToolMenusTests/Tests/ToolMenus/ToolMenusTest.cpp b/Engine/Source/Programs/ToolMenusTests/Tests/ToolMenus/ToolMenusTest.cpp
--- a/Engine/Source/Programs/ToolMenusTests/Tests/ToolMenus/ToolMenusTest.cpp
+++ b/Engine/Source/Programs/ToolMenusTests/Tests/ToolMenus/ToolMenusTest.cpp
@@ -4,48 +4,54 @@
 
 #include "TestHarness.h"
 
-UToolMenus* CreateUniqueUToolMenusInstance()
+namespace UE::ToolMenusTest::Private
 {
-	// Create a unique ToolMenus instance for use in a single test.
-	UToolMenus* ToolMenus = NewObject<UToolMenus>();
-	ToolMenus->AddToRoot();
-	return ToolMenus;
+	UToolMenus* CreateUniqueUToolMenusInstance()
+	{
+		// Create a unique ToolMenus instance for use in a single test.
+		UToolMenus* ToolMenus = NewObject<UToolMenus>();
+		ToolMenus->AddToRoot();
+		return ToolMenus;
+	}
 }
 
 TEST_CASE("Developer::ToolMenus::Can create menu", "[ToolMenus]")
 {
-	UToolMenus* ToolMenus = CreateUniqueUToolMenusInstance();
+	UToolMenus* ToolMenus = UE::ToolMenusTest::Private::CreateUniqueUToolMenusInstance();
 
-	UToolMenu* ToolMenu = ToolMenus->RegisterMenu("MyMenu");
+	UToolMenu* ToolMenu = ToolMenus->RegisterMenu("ToolMenusTest_MyMenu");
 
 	CHECK(ToolMenu);
 }
 
 TEST_CASE("Developer::ToolMenus::Non-registered menu is not registered", "[ToolMenus]")
 {
-	UToolMenus* ToolMenus = CreateUniqueUToolMenusInstance();
+	UToolMenus* ToolMenus = UE::ToolMenusTest::Private::CreateUniqueUToolMenusInstance();
 
-	CHECK_FALSE(ToolMenus->IsMenuRegistered("MyMenu"));
+	CHECK_FALSE(ToolMenus->IsMenuRegistered("ToolMenusTest_MyMenu"));
 }
 
 TEST_CASE("Developer::ToolMenus::Removed menu is not registered", "[ToolMenus]")
 {
-	UToolMenus* ToolMenus = CreateUniqueUToolMenusInstance();
+	UToolMenus* ToolMenus = UE::ToolMenusTest::Private::CreateUniqueUToolMenusInstance();
 
-	ToolMenus->RegisterMenu("MyMenu");
+	const FName MenuName("ToolMenusTest_MyMenu");
 
-	CHECK(ToolMenus->IsMenuRegistered("MyMenu"));
+	ToolMenus->RegisterMenu(MenuName);
 
-	ToolMenus->RemoveMenu("MyMenu");
+	CHECK(ToolMenus->IsMenuRegistered(MenuName));
 
-	CHECK_FALSE(ToolMenus->IsMenuRegistered("MyMenu"));
+	ToolMenus->RemoveMenu(MenuName);
+
+	CHECK_FALSE(ToolMenus->IsMenuRegistered(MenuName));
 }
 
 TEST_CASE("Developer::ToolMenus::GenerateWidget calls dynamic section lambdas", "[ToolMenus]")
 {
-	UToolMenus* ToolMenus = CreateUniqueUToolMenusInstance();
+	UToolMenus* ToolMenus = UE::ToolMenusTest::Private::CreateUniqueUToolMenusInstance();
+	const FName MenuName("ToolMenusTest_MyMenu");
 
-	UToolMenu* ToolMenu = ToolMenus->RegisterMenu("MyMenu");
+	UToolMenu* ToolMenu = ToolMenus->RegisterMenu(MenuName);
 
 	REQUIRE(ToolMenu);
 
@@ -54,17 +60,353 @@ TEST_CASE("Developer::ToolMenus::GenerateWidget calls dynamic section lambdas",
 		bWasLambdaCalled = true;
 	}));
 
-	ToolMenus->GenerateWidget("MyMenu", FToolMenuContext());
+	ToolMenus->GenerateWidget(MenuName, FToolMenuContext());
 
 	CHECK(bWasLambdaCalled);
 }
 
+TEST_CASE("Developer::ToolMenus::Can add and remove sections to menu", "[ToolMenus]")
+{
+	UToolMenus* ToolMenus = UE::ToolMenusTest::Private::CreateUniqueUToolMenusInstance();
+	UToolMenu* ToolMenu = ToolMenus->RegisterMenu("ToolMenusTest_MyMenu");
+
+	REQUIRE(ToolMenu);
+
+	// Add a section to the menu
+	const FName SectionName("MySection");
+	
+	ToolMenu->AddSection(SectionName);
+	CHECK(ToolMenu->Sections.Num() == 1);
+
+	// Check if the section we added is correct
+	FToolMenuSection* Section = ToolMenu->FindSection(SectionName);
+	REQUIRE(Section);
+	CHECK(Section->Name == SectionName);
+
+	// Check if we are able to remove the section
+	ToolMenu->RemoveSection(SectionName);
+	CHECK(ToolMenu->Sections.Num() == 0);
+}
+
+TEST_CASE("Developer::ToolMenus::Can add entries to menu section", "[ToolMenus]")
+{
+
+	UToolMenus* ToolMenus = UE::ToolMenusTest::Private::CreateUniqueUToolMenusInstance();
+	UToolMenu* ToolMenu = ToolMenus->RegisterMenu("ToolMenusTest_MyMenu");
+
+	REQUIRE(ToolMenu);
+
+	FToolMenuSection& Section = ToolMenu->AddSection("MySection");
+
+	// Add the first entry to the menu
+	const FName Entry1Name("MyEntry1");
+	Section.AddMenuEntry(Entry1Name, FText::GetEmpty(), FText::GetEmpty(), FSlateIcon(), FToolUIActionChoice());
+	CHECK(Section.Blocks.Num() == 1);
+
+	// Check if the first entry is correct
+	FToolMenuEntry* Entry1 = Section.FindEntry(Entry1Name);
+	REQUIRE(Entry1);
+	CHECK(Entry1->Name == Entry1Name);
+
+	// Add the second entry to the menu
+	const FName Entry2Name("MyEntry2");
+	Section.AddMenuEntry(Entry2Name, FText::GetEmpty(), FText::GetEmpty(), FSlateIcon(), FToolUIActionChoice());
+	CHECK(Section.Blocks.Num() == 2);
+
+	// Check if the second entry is correct
+	FToolMenuEntry* Entry2 = Section.FindEntry(Entry2Name);
+	REQUIRE(Entry2);
+	CHECK(Entry2->Name == Entry2Name);
+}
+
+TEST_CASE("Developer::ToolMenus::Can add runtime menu customization", "[ToolMenus]")
+{
+	// We have to use the global UToolMenus instance in this test because some of the customization API assumes the menu is registered to that
+	UToolMenus* ToolMenus = UToolMenus::Get();
+
+	// Using a convoluted name for this test because we are adding to the global UToolMenus instance
+	const FName MenuName("ToolMenusTest.CanAddCustomization.MyMenuToUnregister");
+	
+	ToolMenus->RegisterMenu(MenuName);
+	
+	// Add the customization to the menu
+	ToolMenus->AddRuntimeMenuCustomization(MenuName);
+
+	// Try and find the customization
+	FCustomizedToolMenu* CustomizedToolMenu = ToolMenus->FindRuntimeMenuCustomization(MenuName);
+	CHECK(CustomizedToolMenu);
+
+	// Cleanup
+	ToolMenus->RemoveMenu(MenuName);
+}
+
+
+TEST_CASE("Developer::ToolMenus::Can customize menu using runtime menu customization", "[ToolMenus]")
+{
+	// We have to use the global UToolMenus instance in this test because some of the customization API assumes the menu is registered to that
+	UToolMenus* ToolMenus = UToolMenus::Get();
+	const FName MenuName("ToolMenusTest.CanApplyMenuCustomization.MyMenu");
+	
+	UToolMenu* ToolMenu = ToolMenus->RegisterMenu(MenuName);
+
+	REQUIRE(ToolMenu);
+
+	// Add test entries 1, 2 and 3
+	const FName SectionName("MySection");
+	const FName Entry1Name("MyEntry1");
+	const FName Entry2Name("MyEntry2");
+	const FName Entry3Name("MyEntry3");
+
+	FToolMenuSection& Section = ToolMenu->AddSection(SectionName);
+	Section.AddMenuEntry(Entry1Name, FText::GetEmpty(), FText::GetEmpty(), FSlateIcon(), FToolUIActionChoice());
+	Section.AddMenuEntry(Entry2Name, FText::GetEmpty(), FText::GetEmpty(), FSlateIcon(), FToolUIActionChoice());
+	Section.AddMenuEntry(Entry3Name, FText::GetEmpty(), FText::GetEmpty(), FSlateIcon(), FToolUIActionChoice());
+
+	// Add the customization to the menu
+	ToolMenus->AddRuntimeMenuCustomization(MenuName);
+	FCustomizedToolMenu* CustomizedToolMenu = ToolMenus->FindRuntimeMenuCustomization(MenuName);
+	REQUIRE(CustomizedToolMenu);
+
+	// Test the customization by denying entry 2
+	const FName OwnerName("ToolMenusTest");
+
+	CustomizedToolMenu->MenuPermissions.AddDenyListItem(OwnerName, Entry2Name);
+
+	UToolMenu* GeneratedMenu = ToolMenus->GenerateMenu(MenuName, FToolMenuContext());
+	REQUIRE(GeneratedMenu);
+
+	FToolMenuSection* GeneratedMenuSection = GeneratedMenu->FindSection(SectionName);
+	REQUIRE(GeneratedMenuSection);
+
+	// We should only have 2 entries: Entry1 and Entry3 since Entry2 is denied
+	CHECK(GeneratedMenuSection->Blocks.Num() == 2);
+	CHECK(GeneratedMenuSection->FindEntry(Entry1Name));
+	CHECK_FALSE(GeneratedMenuSection->FindEntry(Entry2Name));
+	CHECK(GeneratedMenuSection->FindEntry(Entry3Name));
+
+	// Cleanup
+	ToolMenus->RemoveMenu(MenuName);
+	ToolMenus->UnregisterRuntimeMenuCustomizationOwner(OwnerName);
+}
+
+TEST_CASE("Developer::ToolMenus::Can add runtime menu profile", "[ToolMenus]")
+{
+	// We have to use the global UToolMenus instance in this test because some of the customization API assumes the menu is registered to that
+	UToolMenus* ToolMenus = UToolMenus::Get();
+
+	// Using a convoluted name for this test because we are adding to the global UToolMenus instance
+	const FName MenuName("ToolMenusTest.CanAddProfile.MyMenu");
+
+	ToolMenus->RegisterMenu(MenuName);
+	
+	// Add the customization to the menu
+	const FName ProfileName("MyProfile1");
+	ToolMenus->AddRuntimeMenuProfile(MenuName, ProfileName);
+
+	// Try and find the profile
+	FToolMenuProfile* MenuProfile = ToolMenus->FindRuntimeMenuProfile(MenuName, ProfileName);
+	CHECK(MenuProfile);
+}
+
+TEST_CASE("Developer::ToolMenus::Can customize menu using runtime menu profile", "[ToolMenus]")
+{
+	// We have to use the global UToolMenus instance in this test because some of the profile API assumes the menu is registered to that
+	UToolMenus* ToolMenus = UToolMenus::Get();
+
+	// Using a convoluted name for this test because we are adding to the global UToolMenus instance
+	const FName MenuName("ToolMenusTest.CanApplyMenuProfile.MyMenuToUnregister");
+
+	UToolMenu* ToolMenu = ToolMenus->RegisterMenu(MenuName);
+
+	REQUIRE(ToolMenu);
+
+	// Add test entries 1, 2 and 3 to the menu
+	const FName SectionName("MySection");
+	const FName Entry1Name("MyEntry1");
+	const FName Entry2Name("MyEntry2");
+	const FName Entry3Name("MyEntry3");
+
+	FToolMenuSection& Section = ToolMenu->AddSection(SectionName);
+	Section.AddMenuEntry(Entry1Name, FText::GetEmpty(), FText::GetEmpty(), FSlateIcon(), FToolUIActionChoice());
+	Section.AddMenuEntry(Entry2Name, FText::GetEmpty(), FText::GetEmpty(), FSlateIcon(), FToolUIActionChoice());
+	Section.AddMenuEntry(Entry3Name, FText::GetEmpty(), FText::GetEmpty(), FSlateIcon(), FToolUIActionChoice());
+
+	// Add a profile to the menu
+	const FName Profile1Name("MyProfile1");
+
+	ToolMenus->AddRuntimeMenuProfile(MenuName, Profile1Name);
+
+	// Check if the profile was added
+	FToolMenuProfile* MenuProfile = ToolMenus->FindRuntimeMenuProfile(MenuName, Profile1Name);
+	REQUIRE(MenuProfile);
+
+	// Deny Entry 2 in the profile
+	const FName OwnerName("ToolMenusTest");
+	
+	MenuProfile->MenuPermissions.AddDenyListItem(OwnerName, Entry2Name);
+
+	SECTION("Generating menu without profile should show all entries")
+	{
+		UToolMenu* GeneratedMenuWithoutProfile = ToolMenus->GenerateMenu(MenuName, FToolMenuContext());
+		REQUIRE(GeneratedMenuWithoutProfile);
+
+		FToolMenuSection* GeneratedMenuSection = GeneratedMenuWithoutProfile->FindSection(SectionName);
+		REQUIRE(GeneratedMenuSection);
+
+		CHECK(GeneratedMenuSection->Blocks.Num() == 3);
+		CHECK(GeneratedMenuSection->FindEntry(Entry1Name));
+		CHECK(GeneratedMenuSection->FindEntry(Entry2Name));
+		CHECK(GeneratedMenuSection->FindEntry(Entry3Name));
+	}
+
+	SECTION("Generating menu with profile should hide entry 2")
+	{
+		FToolMenuContext MenuContext;
+
+		UToolMenuProfileContext* ProfileContext = NewObject<UToolMenuProfileContext>();
+		ProfileContext->ActiveProfiles.Add(Profile1Name);
+		MenuContext.AddObject(ProfileContext);
+		
+		UToolMenu* GeneratedMenuWithProfile = ToolMenus->GenerateMenu(MenuName, MenuContext);
+		REQUIRE(GeneratedMenuWithProfile);
+
+		FToolMenuSection* GeneratedMenuSection = GeneratedMenuWithProfile->FindSection(SectionName);
+		REQUIRE(GeneratedMenuSection);
+
+		CHECK(GeneratedMenuSection->Blocks.Num() == 2);
+		CHECK(GeneratedMenuSection->FindEntry(Entry1Name));
+		CHECK_FALSE(GeneratedMenuSection->FindEntry(Entry2Name));
+		CHECK(GeneratedMenuSection->FindEntry(Entry3Name));
+	}
+	
+	// Cleanup
+	ToolMenus->RemoveMenu(MenuName);
+	ToolMenus->UnregisterRuntimeMenuProfileOwner(OwnerName);
+}
+
+TEST_CASE("Developer::ToolMenus::Can combine customization and multiple profiles", "[ToolMenus]")
+{
+	// We have to use the global UToolMenus instance in this test because some of the profile API assumes the menu is registered to that
+	UToolMenus* ToolMenus = UToolMenus::Get();
+
+	// Using a convoluted name for this test because we are adding to the global UToolMenus instance
+	const FName MenuName("ToolMenusTest.CanCombineCustomizationAndProfile.MyMenuToUnregister");
+
+	UToolMenu* ToolMenu = ToolMenus->RegisterMenu(MenuName);
+
+	REQUIRE(ToolMenu);
+
+	// Add the section and entries
+	const FName SectionName("MySection");
+	const FName Entry1Name("MyEntry1");
+	const FName Entry2Name("MyEntry2");
+	const FName Entry3Name("MyEntry3");
+
+	FToolMenuSection& Section = ToolMenu->AddSection(SectionName);
+	Section.AddMenuEntry(Entry1Name, FText::GetEmpty(), FText::GetEmpty(), FSlateIcon(), FToolUIActionChoice());
+	Section.AddMenuEntry(Entry2Name, FText::GetEmpty(), FText::GetEmpty(), FSlateIcon(), FToolUIActionChoice());
+	Section.AddMenuEntry(Entry3Name, FText::GetEmpty(), FText::GetEmpty(), FSlateIcon(), FToolUIActionChoice());
+
+	// The menu customization denies entry 1
+	const FName OwnerName("ToolMenusTest");
+
+	ToolMenus->AddRuntimeMenuCustomization(MenuName);
+	FCustomizedToolMenu* CustomizedToolMenu = ToolMenus->FindRuntimeMenuCustomization(MenuName);
+	REQUIRE(CustomizedToolMenu);
+	CustomizedToolMenu->MenuPermissions.AddDenyListItem(OwnerName, Entry2Name);
+
+	// Add Profile 1 which denies entry 1 and 2
+	const FName Profile1Name("MyProfile1");
+
+	ToolMenus->AddRuntimeMenuProfile(MenuName, Profile1Name);
+	FToolMenuProfile* MenuProfile = ToolMenus->FindRuntimeMenuProfile(MenuName, Profile1Name);
+	REQUIRE(MenuProfile);
+	MenuProfile->MenuPermissions.AddDenyListItem(OwnerName, Entry1Name);
+	MenuProfile->MenuPermissions.AddDenyListItem(OwnerName, Entry2Name);
+
+	// Add Profile 2 which ONLY allows entry 1 (i.e denies ALL other entries)
+	const FName Profile2Name("MyProfile2");
+
+	ToolMenus->AddRuntimeMenuProfile(MenuName, Profile2Name);
+	FToolMenuProfile* MenuProfile2 = ToolMenus->FindRuntimeMenuProfile(MenuName, Profile2Name);
+	REQUIRE(MenuProfile2);
+	MenuProfile2->MenuPermissions.AddAllowListItem(OwnerName, Entry1Name);
+
+
+	SECTION("Generating menu with customization + profile 1 should only show entry 3")
+	{
+		FToolMenuContext MenuContext;
+		
+		UToolMenuProfileContext* ProfileContext = NewObject<UToolMenuProfileContext>();
+		ProfileContext->ActiveProfiles.Add(Profile1Name);
+		MenuContext.AddObject(ProfileContext);
+		
+		UToolMenu* GeneratedMenuWithProfile = ToolMenus->GenerateMenu(MenuName, MenuContext);
+		REQUIRE(GeneratedMenuWithProfile);
+
+		FToolMenuSection* GeneratedMenuSection = GeneratedMenuWithProfile->FindSection(SectionName);
+		REQUIRE(GeneratedMenuSection);
+
+		CHECK(GeneratedMenuSection->Blocks.Num() == 1);
+		CHECK_FALSE(GeneratedMenuSection->FindEntry(Entry1Name));
+		CHECK_FALSE(GeneratedMenuSection->FindEntry(Entry2Name));
+		CHECK(GeneratedMenuSection->FindEntry(Entry3Name));
+	}
+	
+	SECTION("Generating menu with customization + profile 2 should only show entry 1")
+	{
+		FToolMenuContext MenuContext;
+		
+		UToolMenuProfileContext* ProfileContext = NewObject<UToolMenuProfileContext>();
+		ProfileContext->ActiveProfiles.Add(Profile2Name);
+		MenuContext.AddObject(ProfileContext);
+		
+		UToolMenu* GeneratedMenuWithProfile = ToolMenus->GenerateMenu(MenuName, MenuContext);
+		REQUIRE(GeneratedMenuWithProfile);
+
+		FToolMenuSection* GeneratedMenuSection = GeneratedMenuWithProfile->FindSection(SectionName);
+		REQUIRE(GeneratedMenuSection);
+
+		CHECK(GeneratedMenuSection->Blocks.Num() == 1);
+		
+		CHECK(GeneratedMenuSection->FindEntry(Entry1Name));
+		CHECK_FALSE(GeneratedMenuSection->FindEntry(Entry2Name));
+		CHECK_FALSE(GeneratedMenuSection->FindEntry(Entry3Name));
+	}
+	
+	SECTION("Generating menu with customization + profile 1 + profile 2 should only no entries")
+	{
+		FToolMenuContext MenuContext;
+
+		UToolMenuProfileContext* ProfileContext = NewObject<UToolMenuProfileContext>();
+		ProfileContext->ActiveProfiles.Add(Profile1Name);
+		ProfileContext->ActiveProfiles.Add(Profile2Name);
+		MenuContext.AddObject(ProfileContext);
+		
+		UToolMenu* GeneratedMenuWithProfile = ToolMenus->GenerateMenu(MenuName, MenuContext);
+		REQUIRE(GeneratedMenuWithProfile);
+
+		FToolMenuSection* GeneratedMenuSection = GeneratedMenuWithProfile->FindSection(SectionName);
+		REQUIRE(GeneratedMenuSection);
+
+		CHECK(GeneratedMenuSection->Blocks.Num() == 0);
+		CHECK_FALSE(GeneratedMenuSection->FindEntry(Entry1Name));
+		CHECK_FALSE(GeneratedMenuSection->FindEntry(Entry2Name));
+		CHECK_FALSE(GeneratedMenuSection->FindEntry(Entry3Name));
+	}
+
+	// Cleanup
+	ToolMenus->RemoveMenu(MenuName);
+	ToolMenus->UnregisterRuntimeMenuCustomizationOwner(OwnerName);
+	ToolMenus->UnregisterRuntimeMenuProfileOwner(OwnerName);
+}
+
 // Repro of UE-201151.
 TEST_CASE("Developer::ToolMenus::GenerateWidget can handle simultaneous AddReferencedObjects calls in legacy dynamic sections", "[ToolMenus]")
 {
-	UToolMenus* ToolMenus = CreateUniqueUToolMenusInstance();
+	UToolMenus* ToolMenus = UE::ToolMenusTest::Private::CreateUniqueUToolMenusInstance();
 
-	UToolMenu* ToolMenu = ToolMenus->RegisterMenu("MyMenu");
+	UToolMenu* ToolMenu = ToolMenus->RegisterMenu("ToolMenusTest_MyMenu");
 
 	REQUIRE(ToolMenu);
 
