Fix warning in ToolMenusTests

diff --git a/Engine/Source/Programs/ToolMenusTests/Tests/ToolMenus/ToolMenusTest.cpp b/Engine/Source/Programs/ToolMenusTests/Tests/ToolMenus/ToolMenusTest.cpp
--- a/Engine/Source/Programs/ToolMenusTests/Tests/ToolMenus/ToolMenusTest.cpp
+++ b/Engine/Source/Programs/ToolMenusTests/Tests/ToolMenus/ToolMenusTest.cpp
@@ -418,5 +418,5 @@ TEST_CASE("Developer::ToolMenus::GenerateWidget can handle simultaneous AddRefer
 		UToolMenus::AddReferencedObjects(ToolMenus, Finder);
 	}));
 
-	ToolMenus->GenerateWidget("MyMenu", FToolMenuContext());
+	ToolMenus->GenerateWidget("ToolMenusTest_MyMenu", FToolMenuContext());
 }
