Add FToolMenuEntry::GetCheckedState allowing ToolMenu users to inspect more of their menu state

diff --git a/Engine/Source/Developer/ToolMenus/Private/ToolMenuEntry.cpp b/Engine/Source/Developer/ToolMenus/Private/ToolMenuEntry.cpp
--- a/Engine/Source/Developer/ToolMenus/Private/ToolMenuEntry.cpp
+++ b/Engine/Source/Developer/ToolMenus/Private/ToolMenuEntry.cpp
@@ -46,6 +46,34 @@ FToolMenuEntry& FToolMenuEntry::operator=(const FToolMenuEntry&) = default;
 FToolMenuEntry& FToolMenuEntry::operator=(FToolMenuEntry&&) = default;
 PRAGMA_ENABLE_DEPRECATION_WARNINGS
 
+ECheckBoxState FToolMenuEntry::GetCheckState(const FToolMenuContext& InContext) const
+{
+	{
+		TSharedPtr<const FUICommandList> UnusedCommandList;
+		if (const FUIAction* UIAction = GetActionForCommand(InContext, UnusedCommandList))
+		{
+			return UIAction->GetCheckState();
+		}
+	}
+
+	if (const FUIAction* UIAction = Action.GetUIAction())
+	{
+		return UIAction->GetCheckState();
+	}
+
+	if (const FToolUIAction* ToolUIAction = Action.GetToolUIAction())
+	{
+		return ToolUIAction->GetActionCheckState.Execute(InContext);
+	}
+
+	if (const FToolDynamicUIAction* ToolDynamicUIAction = Action.GetToolDynamicUIAction())
+	{
+		return ToolDynamicUIAction->GetActionCheckState.Execute(InContext);
+	}
+
+	return ECheckBoxState::Undetermined;
+}
+
 const FUIAction* FToolMenuEntry::GetActionForCommand(const FToolMenuContext& InContext, TSharedPtr<const FUICommandList>& OutCommandList) const
 {
 	if (Command.IsValid())
diff --git a/Engine/Source/Developer/ToolMenus/Public/ToolMenuEntry.h b/Engine/Source/Developer/ToolMenus/Public/ToolMenuEntry.h
--- a/Engine/Source/Developer/ToolMenus/Public/ToolMenuEntry.h
+++ b/Engine/Source/Developer/ToolMenus/Public/ToolMenuEntry.h
@@ -144,6 +144,12 @@ struct TOOLMENUS_API FToolMenuEntry
 
 	bool IsConstructLegacy() const { return ConstructLegacy.IsBound(); }
 
+	/**
+	 * Get the checked state of this entry by calling underlying commands and delegates.
+	 *
+	 * @return: the checked state of this entry. */
+	ECheckBoxState GetCheckState(const FToolMenuContext& InContext) const;
+
 	const FUIAction* GetActionForCommand(const FToolMenuContext& InContext, TSharedPtr<const FUICommandList>& OutCommandList) const;
 
 	void SetCommandList(const TSharedPtr<const FUICommandList>& InCommandList);
diff --git a/Engine/Source/Programs/ToolMenusTests/Tests/ToolMenus/CheckBoxStateObject.h b/Engine/Source/Programs/ToolMenusTests/Tests/ToolMenus/CheckBoxStateObject.h
new file mode 100644
--- /dev/null
+++ b/Engine/Source/Programs/ToolMenusTests/Tests/ToolMenus/CheckBoxStateObject.h
@@ -0,0 +1,31 @@
+// Copyright Epic Games, Inc. All Rights Reserved.
+
+#pragma once
+
+#include "Styling/SlateTypes.h"
+#include "ToolMenuContext.h"
+
+#include "CheckBoxStateObject.generated.h"
+
+UCLASS()
+class UCheckBoxStateObject : public UObject
+{
+	GENERATED_BODY()
+
+public:
+	UCheckBoxStateObject() = default;
+
+	void SetStateToReturn(const ECheckBoxState InState)
+	{
+		StateToReturn = InState;
+	}
+
+	UFUNCTION()
+	ECheckBoxState GetActionCheckState(const FToolMenuContext& InContext) const
+	{
+		return StateToReturn;
+	}
+
+private:
+	ECheckBoxState StateToReturn;
+};
diff --git a/Engine/Source/Programs/ToolMenusTests/Tests/ToolMenus/ToolMenuEntryTest.cpp b/Engine/Source/Programs/ToolMenusTests/Tests/ToolMenus/ToolMenuEntryTest.cpp
new file mode 100644
--- /dev/null
+++ b/Engine/Source/Programs/ToolMenusTests/Tests/ToolMenus/ToolMenuEntryTest.cpp
@@ -0,0 +1,220 @@
+// Copyright Epic Games, Inc. All Rights Reserved.
+
+#include "ToolMenuEntry.h"
+
+#include "CheckBoxStateObject.h"
+#include "TestHarness.h"
+#include "ToolMenus.h"
+
+TEST_CASE("Developer::ToolMenus::UToolMenuEntry::GetCheckState supports commands", "[ToolMenus]")
+{
+	GIVEN("A command list with a command bound to return Checked")
+	{
+		const TSharedPtr<FUICommandInfo> CommandInfo = MakeShared<FUICommandInfo>("MyTestContext");
+		const TSharedPtr<FUICommandList> CommandList = MakeShared<FUICommandList>();
+		CommandList->MapAction(
+			CommandInfo,
+			FExecuteAction(),
+			FCanExecuteAction(),
+			FGetActionCheckState::CreateLambda(
+				[]() -> ECheckBoxState
+				{
+					return ECheckBoxState::Checked;
+				}
+			)
+		);
+
+		AND_GIVEN("A menu entry initialized with the command and an FToolMenuContext that includes the command list")
+		{
+			FToolMenuContext Context;
+			Context.AppendCommandList(CommandList);
+
+			const FToolMenuEntry Entry = FToolMenuEntry::InitMenuEntry(CommandInfo);
+
+			THEN("FToolMenuEntry::GetCheckState returns Checked")
+			{
+				CHECK(Entry.GetCheckState(Context) == ECheckBoxState::Checked);
+			}
+		}
+
+		AND_GIVEN("A menu entry initialized with both the command list and the command")
+		{
+			FToolMenuContext Context;
+
+			const FToolMenuEntry Entry = FToolMenuEntry::InitMenuEntryWithCommandList(CommandInfo, CommandList);
+
+			THEN("FToolMenuEntry::GetCheckState returns Checked")
+			{
+				CHECK(Entry.GetCheckState(Context) == ECheckBoxState::Checked);
+			}
+		}
+	}
+
+	GIVEN("A command list with a command bound to return Unchecked")
+	{
+		const TSharedPtr<FUICommandInfo> CommandInfo = MakeShared<FUICommandInfo>("MyTestContext");
+		const TSharedPtr<FUICommandList> CommandList = MakeShared<FUICommandList>();
+		CommandList->MapAction(
+			CommandInfo,
+			FExecuteAction(),
+			FCanExecuteAction(),
+			FGetActionCheckState::CreateLambda(
+				[]() -> ECheckBoxState
+				{
+					return ECheckBoxState::Unchecked;
+				}
+			)
+		);
+
+		AND_GIVEN("A menu entry initialized with the command and an FToolMenuContext that includes the command list")
+		{
+			FToolMenuContext Context;
+			Context.AppendCommandList(CommandList);
+
+			const FToolMenuEntry Entry = FToolMenuEntry::InitMenuEntry(CommandInfo);
+
+			THEN("FToolMenuEntry::GetCheckState returns Unchecked")
+			{
+				CHECK(Entry.GetCheckState(Context) == ECheckBoxState::Unchecked);
+			}
+		}
+
+		AND_GIVEN("A menu entry initialized with both the command list and the command")
+		{
+			FToolMenuContext Context;
+
+			const FToolMenuEntry Entry = FToolMenuEntry::InitMenuEntryWithCommandList(CommandInfo, CommandList);
+
+			THEN("FToolMenuEntry::GetCheckState returns Unchecked")
+			{
+				CHECK(Entry.GetCheckState(Context) == ECheckBoxState::Unchecked);
+			}
+		}
+	}
+}
+
+TEST_CASE("Developer::ToolMenus::UToolMenuEntry::GetCheckState supports FUIAction", "[ToolMenus]")
+{
+	GIVEN("An entry with a FUIAction with a GetActionCheckState that returns true")
+	{
+		const FToolMenuEntry Entry = FToolMenuEntry::InitMenuEntry(
+			NAME_None,
+			FText(),
+			FText(),
+			FSlateIcon(),
+			FUIAction(
+				FExecuteAction(),
+				FCanExecuteAction(),
+				FIsActionChecked::CreateLambda(
+					[]() -> bool
+					{
+						return true;
+					}
+				)
+			)
+		);
+
+		THEN("FToolMenuEntry::GetCheckState returns Checked")
+		{
+			CHECK(Entry.GetCheckState(FToolMenuContext()) == ECheckBoxState::Checked);
+		}
+	}
+
+	GIVEN("An entry with a FUIAction with a GetActionCheckState that returns false")
+	{
+		const FToolMenuEntry Entry = FToolMenuEntry::InitMenuEntry(
+			NAME_None,
+			FText(),
+			FText(),
+			FSlateIcon(),
+			FUIAction(
+				FExecuteAction(),
+				FCanExecuteAction(),
+				FIsActionChecked::CreateLambda(
+					[]() -> bool
+					{
+						return false;
+					}
+				)
+			)
+		);
+
+		THEN("FToolMenuEntry::GetCheckState returns Checked")
+		{
+			CHECK(Entry.GetCheckState(FToolMenuContext()) == ECheckBoxState::Unchecked);
+		}
+	}
+}
+
+TEST_CASE("Developer::ToolMenus::UToolMenuEntry::GetCheckState supports FToolUIAction", "[ToolMenus]")
+{
+	GIVEN("An entry with an FToolUIAction with a GetActionCheckState that returns Checked")
+	{
+		FToolUIAction Action;
+		Action.GetActionCheckState = FToolMenuGetActionCheckState::CreateLambda(
+			[](const FToolMenuContext&) -> ECheckBoxState
+			{
+				return ECheckBoxState::Checked;
+			}
+		);
+
+		const FToolMenuEntry Entry = FToolMenuEntry::InitMenuEntry(NAME_None, FText(), FText(), FSlateIcon(), Action);
+
+		THEN("FToolMenuEntry::GetCheckState returns Checked")
+		{
+			CHECK(Entry.GetCheckState(FToolMenuContext()) == ECheckBoxState::Checked);
+		}
+	}
+
+	GIVEN("An entry with an FToolUIAction with a GetActionCheckState that returns Checked")
+	{
+		FToolUIAction Action;
+		Action.GetActionCheckState = FToolMenuGetActionCheckState::CreateLambda(
+			[](const FToolMenuContext&) -> ECheckBoxState
+			{
+				return ECheckBoxState::Unchecked;
+			}
+		);
+
+		const FToolMenuEntry Entry = FToolMenuEntry::InitMenuEntry(NAME_None, FText(), FText(), FSlateIcon(), Action);
+
+		THEN("FToolMenuEntry::GetCheckState returns Unchecked")
+		{
+			CHECK(Entry.GetCheckState(FToolMenuContext()) == ECheckBoxState::Unchecked);
+		}
+	}
+}
+
+TEST_CASE("Developer::ToolMenus::UToolMenuEntry::GetCheckState supports FToolDynamicUIAction", "[ToolMenus]")
+{
+	GIVEN("An entry with an FToolDynamicUIAction with a GetActionCheckState that returns Checked")
+	{
+		UCheckBoxStateObject* CheckBoxStateObject = NewObject<UCheckBoxStateObject>();
+		CheckBoxStateObject->SetStateToReturn(ECheckBoxState::Checked);
+
+		FToolDynamicUIAction Action;
+		Action.GetActionCheckState.BindUFunction(CheckBoxStateObject, "GetActionCheckState");
+
+		const FToolMenuEntry Entry = FToolMenuEntry::InitMenuEntry(NAME_None, FText(), FText(), FSlateIcon(), Action);
+
+		THEN("FToolMenuEntry::GetCheckState returns Checked")
+		{
+			CHECK(Entry.GetCheckState(FToolMenuContext()) == ECheckBoxState::Checked);
+		}
+	}
+	GIVEN("An entry with an FToolDynamicUIAction with a GetActionCheckState that returns Unchecked")
+	{
+		UCheckBoxStateObject* CheckBoxStateObject = NewObject<UCheckBoxStateObject>();
+		CheckBoxStateObject->SetStateToReturn(ECheckBoxState::Unchecked);
+
+		FToolDynamicUIAction Action;
+		Action.GetActionCheckState.BindUFunction(CheckBoxStateObject, "GetActionCheckState");
+
+		const FToolMenuEntry Entry = FToolMenuEntry::InitMenuEntry(NAME_None, FText(), FText(), FSlateIcon(), Action);
+
+		THEN("FToolMenuEntry::GetCheckState returns Unchecked")
+		{
+			CHECK(Entry.GetCheckState(FToolMenuContext()) == ECheckBoxState::Unchecked);
+		}
+	}
+}
