Add FToolMenuTestInstanceScoped to isolate the UToolMenus instances used by tests

* FToolMenuTestInstanceScoped replaces the UToolMenu::Singleton instance when constructed and reinstates the old singleton when destructed.
* Updated all tests to use the new scope. Simplified some menu names in the process since they no longer have to be obscured to avoid collision.
* Remove UE::ToolMenus::Private::CreateUniqueUToolMenusInstance as it is no longer needed.
* Add tests of FToolMenuTestInstanceScoped itself. For example one to verify that it prevents menu names from leaking between tests.

diff --git a/Engine/Source/Developer/ToolMenus/Private/ToolMenus.cpp b/Engine/Source/Developer/ToolMenus/Private/ToolMenus.cpp
--- a/Engine/Source/Developer/ToolMenus/Private/ToolMenus.cpp
+++ b/Engine/Source/Developer/ToolMenus/Private/ToolMenus.cpp
@@ -23,6 +23,19 @@
 
 DEFINE_LOG_CATEGORY(LogToolMenus);
 
+namespace UE::ToolMenus::Private
+{
+
+UToolMenus* CreateToolMenusInstance()
+{
+	UToolMenus* Instance = NewObject<UToolMenus>();
+	Instance->AddToRoot();
+	check(Instance);
+	return Instance;
+}
+
+} // namespace UE::ToolMenus::Private
+
 UToolMenus* UToolMenus::Singleton = nullptr;
 bool UToolMenus::bHasShutDown = false;
 FSimpleMulticastDelegate UToolMenus::StartupCallbacks;
@@ -396,9 +409,7 @@ UToolMenus* UToolMenus::Get()
 		// Required for StartupModule and ShutdownModule to be called and FindModule to list the ToolsMenus module
 		FModuleManager::LoadModuleChecked<IToolMenusModule>("ToolMenus");
 
-		Singleton = NewObject<UToolMenus>();
-		Singleton->AddToRoot();
-		check(Singleton);
+		Singleton = UE::ToolMenus::Private::CreateToolMenusInstance();
 	}
 
 	return Singleton;
@@ -3211,5 +3222,29 @@ void UToolMenus::RemoveAllCustomizations()
 	CustomizedMenus.Reset();
 }
 
+namespace UE::ToolMenus
+{
+
+FToolMenuTestInstanceScoped::FToolMenuTestInstanceScoped()
+{
+	ScopedInstance = Private::CreateToolMenusInstance();
+
+	PreviousInstance = UToolMenus::Get();
+
+	UToolMenus::Singleton = ScopedInstance;
+}
+
+FToolMenuTestInstanceScoped::~FToolMenuTestInstanceScoped()
+{
+	// Reinstate the previous singleton.
+	UToolMenus::Singleton = PreviousInstance;
+
+	// Remove our scoped instance from the root after reinstating the previous instance to not risk our scoped instance
+	// from destroying and taking down the whole ToolMenus system with it.
+	ScopedInstance->RemoveFromRoot();
+}
+
+} // namespace UE::ToolMenus
+
 #undef LOCTEXT_NAMESPACE
 
diff --git a/Engine/Source/Developer/ToolMenus/Public/ToolMenus.h b/Engine/Source/Developer/ToolMenus/Public/ToolMenus.h
--- a/Engine/Source/Developer/ToolMenus/Public/ToolMenus.h
+++ b/Engine/Source/Developer/ToolMenus/Public/ToolMenus.h
@@ -27,6 +27,10 @@ class FMenuBuilder;
 class FMultiBox;
 class SWidget;
 
+namespace UE::ToolMenus
+{
+class FToolMenuTestInstanceScoped;
+}
 class UToolMenu;
 class UToolMenuEntryScript;
 struct FToolMenuEntry;
@@ -468,7 +472,6 @@ private:
 	/** Transient profiles made during runtime that will not be saved */
 	TMap<FName /*MenuName*/, FToolMenuProfileMap> RuntimeMenuProfiles;
 
-
 	FSimpleDelegate SetTimerForNextTickDelegate;
 
 	bool bNextTickTimerIsSet;
@@ -482,6 +485,8 @@ private:
 	static bool bHasShutDown;
 	static FSimpleMulticastDelegate StartupCallbacks;
 	static TOptional<FDelegateHandle> InternalStartupCallbackHandle;
+
+	friend class UE::ToolMenus::FToolMenuTestInstanceScoped;
 };
 
 /**
@@ -506,3 +511,23 @@ private:
 
 	FToolMenuOwner Owner;
 };
+
+namespace UE::ToolMenus
+{
+
+/**
+ * Creates a new instance of UToolMenus and sets the UToolMenus singleton to point at this new instance.
+ * When the scope instance is destructed, the previous UToolMenus singleton is reinstated.
+ */
+class FToolMenuTestInstanceScoped
+{
+public:
+	FToolMenuTestInstanceScoped();
+	~FToolMenuTestInstanceScoped();
+
+private:
+	UToolMenus* PreviousInstance;
+	UToolMenus* ScopedInstance;
+};
+
+}; // namespace UE::ToolMenus
diff --git a/Engine/Source/Programs/ToolMenusTests/Tests/ToolMenus/ToolMenuTestInstanceScopedTest.cpp b/Engine/Source/Programs/ToolMenusTests/Tests/ToolMenus/ToolMenuTestInstanceScopedTest.cpp
new file mode 100644
--- /dev/null
+++ b/Engine/Source/Programs/ToolMenusTests/Tests/ToolMenus/ToolMenuTestInstanceScopedTest.cpp
@@ -0,0 +1,67 @@
+// Copyright Epic Games, Inc. All Rights Reserved.
+
+#include "TestHarness.h"
+
+#include "ToolMenus.h"
+
+TEST_CASE("UE::ToolMenus::FToolMenuTestInstanceScoped replaces the singleton while the scope is active", "[ToolMenus]")
+{
+	UToolMenus* const OriginalSingleton = UToolMenus::Get();
+
+	{
+		UE::ToolMenus::FToolMenuTestInstanceScoped Scope1;
+
+		UToolMenus* const FirstNestedSingleton = UToolMenus::Get();
+
+		CHECK(FirstNestedSingleton != OriginalSingleton);
+
+		{
+			UE::ToolMenus::FToolMenuTestInstanceScoped Scope2;
+
+			UToolMenus* const SecondNestedSingleton = UToolMenus::Get();
+
+			CHECK(SecondNestedSingleton != FirstNestedSingleton);
+			CHECK(SecondNestedSingleton != OriginalSingleton);
+		}
+	}
+}
+
+TEST_CASE("UE::ToolMenus::FToolMenuTestInstanceScoped reinstates the original singleton when the scope is destructed", "[ToolMenus]")
+{
+	UToolMenus* const OriginalSingleton = UToolMenus::Get();
+
+	{
+		UE::ToolMenus::FToolMenuTestInstanceScoped Scope1;
+
+		UToolMenus* const FirstNestedSingleton = UToolMenus::Get();
+
+		{
+			UE::ToolMenus::FToolMenuTestInstanceScoped Scope2;
+		}
+
+		CHECK(FirstNestedSingleton == UToolMenus::Get());
+	}
+
+	CHECK(OriginalSingleton == UToolMenus::Get());
+}
+
+const FName FToolMenuTestInstanceScopedMenuName =
+	"FToolMenuTestInstanceScoped_prevents_names_from_leaking_between_tests";
+
+TEST_CASE("UE::ToolMenus::FToolMenuTestInstanceScoped prevents menu names from leaking between tests (test 1 of 2)", "[ToolMenus]")
+{
+	UE::ToolMenus::FToolMenuTestInstanceScoped Scope;
+
+	CHECK(!UToolMenus::Get()->IsMenuRegistered(FToolMenuTestInstanceScopedMenuName));
+
+	UToolMenus::Get()->RegisterMenu(FToolMenuTestInstanceScopedMenuName);
+}
+
+TEST_CASE("UE::ToolMenus::FToolMenuTestInstanceScoped prevents menu names from leaking between tests (test 2 of 2)", "[ToolMenus]")
+{
+	UE::ToolMenus::FToolMenuTestInstanceScoped Scope;
+
+	CHECK(!UToolMenus::Get()->IsMenuRegistered(FToolMenuTestInstanceScopedMenuName));
+
+	UToolMenus::Get()->RegisterMenu(FToolMenuTestInstanceScopedMenuName);
+}
diff --git a/Engine/Source/Programs/ToolMenusTests/Tests/ToolMenus/ToolMenusTest.cpp b/Engine/Source/Programs/ToolMenusTests/Tests/ToolMenus/ToolMenusTest.cpp
--- a/Engine/Source/Programs/ToolMenusTests/Tests/ToolMenus/ToolMenusTest.cpp
+++ b/Engine/Source/Programs/ToolMenusTests/Tests/ToolMenus/ToolMenusTest.cpp
@@ -4,54 +4,44 @@
 
 #include "TestHarness.h"
 
-namespace UE::ToolMenusTest::Private
-{
-	UToolMenus* CreateUniqueUToolMenusInstance()
-	{
-		// Create a unique ToolMenus instance for use in a single test.
-		UToolMenus* ToolMenus = NewObject<UToolMenus>();
-		ToolMenus->AddToRoot();
-		return ToolMenus;
-	}
-}
-
 TEST_CASE("Developer::ToolMenus::Can create menu", "[ToolMenus]")
 {
-	UToolMenus* ToolMenus = UE::ToolMenusTest::Private::CreateUniqueUToolMenusInstance();
+	UE::ToolMenus::FToolMenuTestInstanceScoped Scope;
 
-	UToolMenu* ToolMenu = ToolMenus->RegisterMenu("ToolMenusTest_MyMenu");
+	UToolMenu* ToolMenu = UToolMenus::Get()->RegisterMenu("ToolMenusTest_MyMenu");
 
 	CHECK(ToolMenu);
 }
 
 TEST_CASE("Developer::ToolMenus::Non-registered menu is not registered", "[ToolMenus]")
 {
-	UToolMenus* ToolMenus = UE::ToolMenusTest::Private::CreateUniqueUToolMenusInstance();
+	UE::ToolMenus::FToolMenuTestInstanceScoped Scope;
 
-	CHECK_FALSE(ToolMenus->IsMenuRegistered("ToolMenusTest_MyMenu"));
+	CHECK_FALSE(UToolMenus::Get()->IsMenuRegistered("ToolMenusTest_MyMenu"));
 }
 
 TEST_CASE("Developer::ToolMenus::Removed menu is not registered", "[ToolMenus]")
 {
-	UToolMenus* ToolMenus = UE::ToolMenusTest::Private::CreateUniqueUToolMenusInstance();
+	UE::ToolMenus::FToolMenuTestInstanceScoped Scope;
 
 	const FName MenuName("ToolMenusTest_MyMenu");
 
-	ToolMenus->RegisterMenu(MenuName);
+	UToolMenus::Get()->RegisterMenu(MenuName);
 
-	CHECK(ToolMenus->IsMenuRegistered(MenuName));
+	CHECK(UToolMenus::Get()->IsMenuRegistered(MenuName));
 
-	ToolMenus->RemoveMenu(MenuName);
+	UToolMenus::Get()->RemoveMenu(MenuName);
 
-	CHECK_FALSE(ToolMenus->IsMenuRegistered(MenuName));
+	CHECK_FALSE(UToolMenus::Get()->IsMenuRegistered(MenuName));
 }
 
 TEST_CASE("Developer::ToolMenus::GenerateWidget calls dynamic section lambdas", "[ToolMenus]")
 {
-	UToolMenus* ToolMenus = UE::ToolMenusTest::Private::CreateUniqueUToolMenusInstance();
+	UE::ToolMenus::FToolMenuTestInstanceScoped Scope;
+
 	const FName MenuName("ToolMenusTest_MyMenu");
 
-	UToolMenu* ToolMenu = ToolMenus->RegisterMenu(MenuName);
+	UToolMenu* ToolMenu = UToolMenus::Get()->RegisterMenu(MenuName);
 
 	REQUIRE(ToolMenu);
 
@@ -60,15 +50,16 @@ TEST_CASE("Developer::ToolMenus::GenerateWidget calls dynamic section lambdas",
 		bWasLambdaCalled = true;
 	}));
 
-	ToolMenus->GenerateWidget(MenuName, FToolMenuContext());
+	UToolMenus::Get()->GenerateWidget(MenuName, FToolMenuContext());
 
 	CHECK(bWasLambdaCalled);
 }
 
 TEST_CASE("Developer::ToolMenus::Can add and remove sections to menu", "[ToolMenus]")
 {
-	UToolMenus* ToolMenus = UE::ToolMenusTest::Private::CreateUniqueUToolMenusInstance();
-	UToolMenu* ToolMenu = ToolMenus->RegisterMenu("ToolMenusTest_MyMenu");
+	UE::ToolMenus::FToolMenuTestInstanceScoped Scope;
+
+	UToolMenu* ToolMenu = UToolMenus::Get()->RegisterMenu("ToolMenusTest_MyMenu");
 
 	REQUIRE(ToolMenu);
 
@@ -90,9 +81,9 @@ TEST_CASE("Developer::ToolMenus::Can add and remove sections to menu", "[ToolMen
 
 TEST_CASE("Developer::ToolMenus::Can add entries to menu section", "[ToolMenus]")
 {
+	UE::ToolMenus::FToolMenuTestInstanceScoped Scope;
 
-	UToolMenus* ToolMenus = UE::ToolMenusTest::Private::CreateUniqueUToolMenusInstance();
-	UToolMenu* ToolMenu = ToolMenus->RegisterMenu("ToolMenusTest_MyMenu");
+	UToolMenu* ToolMenu = UToolMenus::Get()->RegisterMenu("ToolMenusTest_MyMenu");
 
 	REQUIRE(ToolMenu);
 
@@ -121,33 +112,28 @@ TEST_CASE("Developer::ToolMenus::Can add entries to menu section", "[ToolMenus]"
 
 TEST_CASE("Developer::ToolMenus::Can add runtime menu customization", "[ToolMenus]")
 {
-	// We have to use the global UToolMenus instance in this test because some of the customization API assumes the menu is registered to that
-	UToolMenus* ToolMenus = UToolMenus::Get();
+	UE::ToolMenus::FToolMenuTestInstanceScoped Scope;
+
+	const FName MenuName("MyMenu");
+
+	UToolMenus::Get()->RegisterMenu(MenuName);
 
-	// Using a convoluted name for this test because we are adding to the global UToolMenus instance
-	const FName MenuName("ToolMenusTest.CanAddCustomization.MyMenuToUnregister");
-	
-	ToolMenus->RegisterMenu(MenuName);
-	
 	// Add the customization to the menu
-	ToolMenus->AddRuntimeMenuCustomization(MenuName);
+	UToolMenus::Get()->AddRuntimeMenuCustomization(MenuName);
 
 	// Try and find the customization
-	FCustomizedToolMenu* CustomizedToolMenu = ToolMenus->FindRuntimeMenuCustomization(MenuName);
+	FCustomizedToolMenu* CustomizedToolMenu = UToolMenus::Get()->FindRuntimeMenuCustomization(MenuName);
 	CHECK(CustomizedToolMenu);
-
-	// Cleanup
-	ToolMenus->RemoveMenu(MenuName);
 }
 
 
 TEST_CASE("Developer::ToolMenus::Can customize menu using runtime menu customization", "[ToolMenus]")
 {
-	// We have to use the global UToolMenus instance in this test because some of the customization API assumes the menu is registered to that
-	UToolMenus* ToolMenus = UToolMenus::Get();
-	const FName MenuName("ToolMenusTest.CanApplyMenuCustomization.MyMenu");
-	
-	UToolMenu* ToolMenu = ToolMenus->RegisterMenu(MenuName);
+	UE::ToolMenus::FToolMenuTestInstanceScoped Scope;
+
+	const FName MenuName("MyMenu");
+
+	UToolMenu* ToolMenu = UToolMenus::Get()->RegisterMenu(MenuName);
 
 	REQUIRE(ToolMenu);
 
@@ -163,8 +149,8 @@ TEST_CASE("Developer::ToolMenus::Can customize menu using runtime menu customiza
 	Section.AddMenuEntry(Entry3Name, FText::GetEmpty(), FText::GetEmpty(), FSlateIcon(), FToolUIActionChoice());
 
 	// Add the customization to the menu
-	ToolMenus->AddRuntimeMenuCustomization(MenuName);
-	FCustomizedToolMenu* CustomizedToolMenu = ToolMenus->FindRuntimeMenuCustomization(MenuName);
+	UToolMenus::Get()->AddRuntimeMenuCustomization(MenuName);
+	FCustomizedToolMenu* CustomizedToolMenu = UToolMenus::Get()->FindRuntimeMenuCustomization(MenuName);
 	REQUIRE(CustomizedToolMenu);
 
 	// Test the customization by denying entry 2
@@ -172,7 +158,7 @@ TEST_CASE("Developer::ToolMenus::Can customize menu using runtime menu customiza
 
 	CustomizedToolMenu->MenuPermissions.AddDenyListItem(OwnerName, Entry2Name);
 
-	UToolMenu* GeneratedMenu = ToolMenus->GenerateMenu(MenuName, FToolMenuContext());
+	UToolMenu* GeneratedMenu = UToolMenus::Get()->GenerateMenu(MenuName, FToolMenuContext());
 	REQUIRE(GeneratedMenu);
 
 	FToolMenuSection* GeneratedMenuSection = GeneratedMenu->FindSection(SectionName);
@@ -183,40 +169,32 @@ TEST_CASE("Developer::ToolMenus::Can customize menu using runtime menu customiza
 	CHECK(GeneratedMenuSection->FindEntry(Entry1Name));
 	CHECK_FALSE(GeneratedMenuSection->FindEntry(Entry2Name));
 	CHECK(GeneratedMenuSection->FindEntry(Entry3Name));
-
-	// Cleanup
-	ToolMenus->RemoveMenu(MenuName);
-	ToolMenus->UnregisterRuntimeMenuCustomizationOwner(OwnerName);
 }
 
 TEST_CASE("Developer::ToolMenus::Can add runtime menu profile", "[ToolMenus]")
 {
-	// We have to use the global UToolMenus instance in this test because some of the customization API assumes the menu is registered to that
-	UToolMenus* ToolMenus = UToolMenus::Get();
+	UE::ToolMenus::FToolMenuTestInstanceScoped Scope;
 
-	// Using a convoluted name for this test because we are adding to the global UToolMenus instance
-	const FName MenuName("ToolMenusTest.CanAddProfile.MyMenu");
+	const FName MenuName("MyMenu");
+
+	UToolMenus::Get()->RegisterMenu(MenuName);
 
-	ToolMenus->RegisterMenu(MenuName);
-	
 	// Add the customization to the menu
 	const FName ProfileName("MyProfile1");
-	ToolMenus->AddRuntimeMenuProfile(MenuName, ProfileName);
+	UToolMenus::Get()->AddRuntimeMenuProfile(MenuName, ProfileName);
 
 	// Try and find the profile
-	FToolMenuProfile* MenuProfile = ToolMenus->FindRuntimeMenuProfile(MenuName, ProfileName);
+	FToolMenuProfile* MenuProfile = UToolMenus::Get()->FindRuntimeMenuProfile(MenuName, ProfileName);
 	CHECK(MenuProfile);
 }
 
 TEST_CASE("Developer::ToolMenus::Can customize menu using runtime menu profile", "[ToolMenus]")
 {
-	// We have to use the global UToolMenus instance in this test because some of the profile API assumes the menu is registered to that
-	UToolMenus* ToolMenus = UToolMenus::Get();
+	UE::ToolMenus::FToolMenuTestInstanceScoped Scope;
 
-	// Using a convoluted name for this test because we are adding to the global UToolMenus instance
-	const FName MenuName("ToolMenusTest.CanApplyMenuProfile.MyMenuToUnregister");
+	const FName MenuName("MyMenu");
 
-	UToolMenu* ToolMenu = ToolMenus->RegisterMenu(MenuName);
+	UToolMenu* ToolMenu = UToolMenus::Get()->RegisterMenu(MenuName);
 
 	REQUIRE(ToolMenu);
 
@@ -234,10 +212,10 @@ TEST_CASE("Developer::ToolMenus::Can customize menu using runtime menu profile",
 	// Add a profile to the menu
 	const FName Profile1Name("MyProfile1");
 
-	ToolMenus->AddRuntimeMenuProfile(MenuName, Profile1Name);
+	UToolMenus::Get()->AddRuntimeMenuProfile(MenuName, Profile1Name);
 
 	// Check if the profile was added
-	FToolMenuProfile* MenuProfile = ToolMenus->FindRuntimeMenuProfile(MenuName, Profile1Name);
+	FToolMenuProfile* MenuProfile = UToolMenus::Get()->FindRuntimeMenuProfile(MenuName, Profile1Name);
 	REQUIRE(MenuProfile);
 
 	// Deny Entry 2 in the profile
@@ -247,7 +225,7 @@ TEST_CASE("Developer::ToolMenus::Can customize menu using runtime menu profile",
 
 	SECTION("Generating menu without profile should show all entries")
 	{
-		UToolMenu* GeneratedMenuWithoutProfile = ToolMenus->GenerateMenu(MenuName, FToolMenuContext());
+		UToolMenu* GeneratedMenuWithoutProfile = UToolMenus::Get()->GenerateMenu(MenuName, FToolMenuContext());
 		REQUIRE(GeneratedMenuWithoutProfile);
 
 		FToolMenuSection* GeneratedMenuSection = GeneratedMenuWithoutProfile->FindSection(SectionName);
@@ -266,8 +244,8 @@ TEST_CASE("Developer::ToolMenus::Can customize menu using runtime menu profile",
 		UToolMenuProfileContext* ProfileContext = NewObject<UToolMenuProfileContext>();
 		ProfileContext->ActiveProfiles.Add(Profile1Name);
 		MenuContext.AddObject(ProfileContext);
-		
-		UToolMenu* GeneratedMenuWithProfile = ToolMenus->GenerateMenu(MenuName, MenuContext);
+
+		UToolMenu* GeneratedMenuWithProfile = UToolMenus::Get()->GenerateMenu(MenuName, MenuContext);
 		REQUIRE(GeneratedMenuWithProfile);
 
 		FToolMenuSection* GeneratedMenuSection = GeneratedMenuWithProfile->FindSection(SectionName);
@@ -278,21 +256,15 @@ TEST_CASE("Developer::ToolMenus::Can customize menu using runtime menu profile",
 		CHECK_FALSE(GeneratedMenuSection->FindEntry(Entry2Name));
 		CHECK(GeneratedMenuSection->FindEntry(Entry3Name));
 	}
-	
-	// Cleanup
-	ToolMenus->RemoveMenu(MenuName);
-	ToolMenus->UnregisterRuntimeMenuProfileOwner(OwnerName);
 }
 
 TEST_CASE("Developer::ToolMenus::Can combine customization and multiple profiles", "[ToolMenus]")
 {
-	// We have to use the global UToolMenus instance in this test because some of the profile API assumes the menu is registered to that
-	UToolMenus* ToolMenus = UToolMenus::Get();
+	UE::ToolMenus::FToolMenuTestInstanceScoped Scope;
 
-	// Using a convoluted name for this test because we are adding to the global UToolMenus instance
-	const FName MenuName("ToolMenusTest.CanCombineCustomizationAndProfile.MyMenuToUnregister");
+	const FName MenuName("MyMenu");
 
-	UToolMenu* ToolMenu = ToolMenus->RegisterMenu(MenuName);
+	UToolMenu* ToolMenu = UToolMenus::Get()->RegisterMenu(MenuName);
 
 	REQUIRE(ToolMenu);
 
@@ -310,16 +282,16 @@ TEST_CASE("Developer::ToolMenus::Can combine customization and multiple profiles
 	// The menu customization denies entry 1
 	const FName OwnerName("ToolMenusTest");
 
-	ToolMenus->AddRuntimeMenuCustomization(MenuName);
-	FCustomizedToolMenu* CustomizedToolMenu = ToolMenus->FindRuntimeMenuCustomization(MenuName);
+	UToolMenus::Get()->AddRuntimeMenuCustomization(MenuName);
+	FCustomizedToolMenu* CustomizedToolMenu = UToolMenus::Get()->FindRuntimeMenuCustomization(MenuName);
 	REQUIRE(CustomizedToolMenu);
 	CustomizedToolMenu->MenuPermissions.AddDenyListItem(OwnerName, Entry2Name);
 
 	// Add Profile 1 which denies entry 1 and 2
 	const FName Profile1Name("MyProfile1");
 
-	ToolMenus->AddRuntimeMenuProfile(MenuName, Profile1Name);
-	FToolMenuProfile* MenuProfile = ToolMenus->FindRuntimeMenuProfile(MenuName, Profile1Name);
+	UToolMenus::Get()->AddRuntimeMenuProfile(MenuName, Profile1Name);
+	FToolMenuProfile* MenuProfile = UToolMenus::Get()->FindRuntimeMenuProfile(MenuName, Profile1Name);
 	REQUIRE(MenuProfile);
 	MenuProfile->MenuPermissions.AddDenyListItem(OwnerName, Entry1Name);
 	MenuProfile->MenuPermissions.AddDenyListItem(OwnerName, Entry2Name);
@@ -327,8 +299,8 @@ TEST_CASE("Developer::ToolMenus::Can combine customization and multiple profiles
 	// Add Profile 2 which ONLY allows entry 1 (i.e denies ALL other entries)
 	const FName Profile2Name("MyProfile2");
 
-	ToolMenus->AddRuntimeMenuProfile(MenuName, Profile2Name);
-	FToolMenuProfile* MenuProfile2 = ToolMenus->FindRuntimeMenuProfile(MenuName, Profile2Name);
+	UToolMenus::Get()->AddRuntimeMenuProfile(MenuName, Profile2Name);
+	FToolMenuProfile* MenuProfile2 = UToolMenus::Get()->FindRuntimeMenuProfile(MenuName, Profile2Name);
 	REQUIRE(MenuProfile2);
 	MenuProfile2->MenuPermissions.AddAllowListItem(OwnerName, Entry1Name);
 
@@ -340,8 +312,8 @@ TEST_CASE("Developer::ToolMenus::Can combine customization and multiple profiles
 		UToolMenuProfileContext* ProfileContext = NewObject<UToolMenuProfileContext>();
 		ProfileContext->ActiveProfiles.Add(Profile1Name);
 		MenuContext.AddObject(ProfileContext);
-		
-		UToolMenu* GeneratedMenuWithProfile = ToolMenus->GenerateMenu(MenuName, MenuContext);
+
+		UToolMenu* GeneratedMenuWithProfile = UToolMenus::Get()->GenerateMenu(MenuName, MenuContext);
 		REQUIRE(GeneratedMenuWithProfile);
 
 		FToolMenuSection* GeneratedMenuSection = GeneratedMenuWithProfile->FindSection(SectionName);
@@ -360,8 +332,8 @@ TEST_CASE("Developer::ToolMenus::Can combine customization and multiple profiles
 		UToolMenuProfileContext* ProfileContext = NewObject<UToolMenuProfileContext>();
 		ProfileContext->ActiveProfiles.Add(Profile2Name);
 		MenuContext.AddObject(ProfileContext);
-		
-		UToolMenu* GeneratedMenuWithProfile = ToolMenus->GenerateMenu(MenuName, MenuContext);
+
+		UToolMenu* GeneratedMenuWithProfile = UToolMenus::Get()->GenerateMenu(MenuName, MenuContext);
 		REQUIRE(GeneratedMenuWithProfile);
 
 		FToolMenuSection* GeneratedMenuSection = GeneratedMenuWithProfile->FindSection(SectionName);
@@ -382,8 +354,8 @@ TEST_CASE("Developer::ToolMenus::Can combine customization and multiple profiles
 		ProfileContext->ActiveProfiles.Add(Profile1Name);
 		ProfileContext->ActiveProfiles.Add(Profile2Name);
 		MenuContext.AddObject(ProfileContext);
-		
-		UToolMenu* GeneratedMenuWithProfile = ToolMenus->GenerateMenu(MenuName, MenuContext);
+
+		UToolMenu* GeneratedMenuWithProfile = UToolMenus::Get()->GenerateMenu(MenuName, MenuContext);
 		REQUIRE(GeneratedMenuWithProfile);
 
 		FToolMenuSection* GeneratedMenuSection = GeneratedMenuWithProfile->FindSection(SectionName);
@@ -394,29 +366,30 @@ TEST_CASE("Developer::ToolMenus::Can combine customization and multiple profiles
 		CHECK_FALSE(GeneratedMenuSection->FindEntry(Entry2Name));
 		CHECK_FALSE(GeneratedMenuSection->FindEntry(Entry3Name));
 	}
-
-	// Cleanup
-	ToolMenus->RemoveMenu(MenuName);
-	ToolMenus->UnregisterRuntimeMenuCustomizationOwner(OwnerName);
-	ToolMenus->UnregisterRuntimeMenuProfileOwner(OwnerName);
 }
 
 // Repro of UE-201151.
 TEST_CASE("Developer::ToolMenus::GenerateWidget can handle simultaneous AddReferencedObjects calls in legacy dynamic sections", "[ToolMenus]")
 {
-	UToolMenus* ToolMenus = UE::ToolMenusTest::Private::CreateUniqueUToolMenusInstance();
+	UE::ToolMenus::FToolMenuTestInstanceScoped Scope;
 
-	UToolMenu* ToolMenu = ToolMenus->RegisterMenu("ToolMenusTest_MyMenu");
+	UToolMenu* ToolMenu = UToolMenus::Get()->RegisterMenu("ToolMenusTest_MyMenu");
 
 	REQUIRE(ToolMenu);
 
 	// This simulates the crash of UE-201151 that occurred in UToolMenus::GenerateWidget(UToolMenu*) after a complex delegate
 	// triggered a call to UToolMenus::AddReferencedObjects while the delegate was still executing.
-	ToolMenu->AddDynamicSection("MyDynamicLegacySection", FNewToolMenuDelegateLegacy::CreateLambda([ToolMenus](FMenuBuilder&, UToolMenu*) {
-		TArray<UObject*> Array;
-		FReferenceFinder Finder(Array);
-		UToolMenus::AddReferencedObjects(ToolMenus, Finder);
-	}));
-
-	ToolMenus->GenerateWidget("ToolMenusTest_MyMenu", FToolMenuContext());
+	ToolMenu->AddDynamicSection(
+		"MyDynamicLegacySection",
+		FNewToolMenuDelegateLegacy::CreateLambda(
+			[](FMenuBuilder&, UToolMenu*)
+			{
+				TArray<UObject*> Array;
+				FReferenceFinder Finder(Array);
+				UToolMenus::AddReferencedObjects(UToolMenus::Get(), Finder);
+			}
+		)
+	);
+
+	UToolMenus::Get()->GenerateWidget("ToolMenusTest_MyMenu", FToolMenuContext());
 }
