Restructure files of ToolMenusTests to avoid a warning at build time

diff --git a/Engine/Source/Programs/ToolMenusTests/Private/CheckBoxStateObject.h b/Engine/Source/Programs/ToolMenusTests/Private/CheckBoxStateObject.h
new file mode 100644
--- /dev/null
+++ b/Engine/Source/Programs/ToolMenusTests/Private/CheckBoxStateObject.h
@@ -0,0 +1,31 @@
+// Copyright Epic Games, Inc. All Rights Reserved.
+
+#pragma once
+
+#include "Styling/SlateTypes.h"
+#include "ToolMenuContext.h"
+
+#include "CheckBoxStateObject.generated.h"
+
+UCLASS()
+class UCheckBoxStateObject : public UObject
+{
+	GENERATED_BODY()
+
+public:
+	UCheckBoxStateObject() = default;
+
+	void SetStateToReturn(const ECheckBoxState InState)
+	{
+		StateToReturn = InState;
+	}
+
+	UFUNCTION()
+	ECheckBoxState GetActionCheckState(const FToolMenuContext& InContext) const
+	{
+		return StateToReturn;
+	}
+
+private:
+	ECheckBoxState StateToReturn;
+};
diff --git a/Engine/Source/Programs/ToolMenusTests/Tests/ToolMenuEntryTest.cpp b/Engine/Source/Programs/ToolMenusTests/Tests/ToolMenuEntryTest.cpp
new file mode 100644
--- /dev/null
+++ b/Engine/Source/Programs/ToolMenusTests/Tests/ToolMenuEntryTest.cpp
@@ -0,0 +1,220 @@
+// Copyright Epic Games, Inc. All Rights Reserved.
+
+#include "ToolMenuEntry.h"
+
+#include "CheckBoxStateObject.h"
+#include "TestHarness.h"
+#include "ToolMenus.h"
+
+TEST_CASE("Developer::ToolMenus::UToolMenuEntry::GetCheckState supports commands", "[ToolMenus]")
+{
+	GIVEN("A command list with a command bound to return Checked")
+	{
+		const TSharedPtr<FUICommandInfo> CommandInfo = MakeShared<FUICommandInfo>("MyTestContext");
+		const TSharedPtr<FUICommandList> CommandList = MakeShared<FUICommandList>();
+		CommandList->MapAction(
+			CommandInfo,
+			FExecuteAction(),
+			FCanExecuteAction(),
+			FGetActionCheckState::CreateLambda(
+				[]() -> ECheckBoxState
+				{
+					return ECheckBoxState::Checked;
+				}
+			)
+		);
+
+		AND_GIVEN("A menu entry initialized with the command and an FToolMenuContext that includes the command list")
+		{
+			FToolMenuContext Context;
+			Context.AppendCommandList(CommandList);
+
+			const FToolMenuEntry Entry = FToolMenuEntry::InitMenuEntry(CommandInfo);
+
+			THEN("FToolMenuEntry::GetCheckState returns Checked")
+			{
+				CHECK(Entry.GetCheckState(Context) == ECheckBoxState::Checked);
+			}
+		}
+
+		AND_GIVEN("A menu entry initialized with both the command list and the command")
+		{
+			FToolMenuContext Context;
+
+			const FToolMenuEntry Entry = FToolMenuEntry::InitMenuEntryWithCommandList(CommandInfo, CommandList);
+
+			THEN("FToolMenuEntry::GetCheckState returns Checked")
+			{
+				CHECK(Entry.GetCheckState(Context) == ECheckBoxState::Checked);
+			}
+		}
+	}
+
+	GIVEN("A command list with a command bound to return Unchecked")
+	{
+		const TSharedPtr<FUICommandInfo> CommandInfo = MakeShared<FUICommandInfo>("MyTestContext");
+		const TSharedPtr<FUICommandList> CommandList = MakeShared<FUICommandList>();
+		CommandList->MapAction(
+			CommandInfo,
+			FExecuteAction(),
+			FCanExecuteAction(),
+			FGetActionCheckState::CreateLambda(
+				[]() -> ECheckBoxState
+				{
+					return ECheckBoxState::Unchecked;
+				}
+			)
+		);
+
+		AND_GIVEN("A menu entry initialized with the command and an FToolMenuContext that includes the command list")
+		{
+			FToolMenuContext Context;
+			Context.AppendCommandList(CommandList);
+
+			const FToolMenuEntry Entry = FToolMenuEntry::InitMenuEntry(CommandInfo);
+
+			THEN("FToolMenuEntry::GetCheckState returns Unchecked")
+			{
+				CHECK(Entry.GetCheckState(Context) == ECheckBoxState::Unchecked);
+			}
+		}
+
+		AND_GIVEN("A menu entry initialized with both the command list and the command")
+		{
+			FToolMenuContext Context;
+
+			const FToolMenuEntry Entry = FToolMenuEntry::InitMenuEntryWithCommandList(CommandInfo, CommandList);
+
+			THEN("FToolMenuEntry::GetCheckState returns Unchecked")
+			{
+				CHECK(Entry.GetCheckState(Context) == ECheckBoxState::Unchecked);
+			}
+		}
+	}
+}
+
+TEST_CASE("Developer::ToolMenus::UToolMenuEntry::GetCheckState supports FUIAction", "[ToolMenus]")
+{
+	GIVEN("An entry with a FUIAction with a GetActionCheckState that returns true")
+	{
+		const FToolMenuEntry Entry = FToolMenuEntry::InitMenuEntry(
+			NAME_None,
+			FText(),
+			FText(),
+			FSlateIcon(),
+			FUIAction(
+				FExecuteAction(),
+				FCanExecuteAction(),
+				FIsActionChecked::CreateLambda(
+					[]() -> bool
+					{
+						return true;
+					}
+				)
+			)
+		);
+
+		THEN("FToolMenuEntry::GetCheckState returns Checked")
+		{
+			CHECK(Entry.GetCheckState(FToolMenuContext()) == ECheckBoxState::Checked);
+		}
+	}
+
+	GIVEN("An entry with a FUIAction with a GetActionCheckState that returns false")
+	{
+		const FToolMenuEntry Entry = FToolMenuEntry::InitMenuEntry(
+			NAME_None,
+			FText(),
+			FText(),
+			FSlateIcon(),
+			FUIAction(
+				FExecuteAction(),
+				FCanExecuteAction(),
+				FIsActionChecked::CreateLambda(
+					[]() -> bool
+					{
+						return false;
+					}
+				)
+			)
+		);
+
+		THEN("FToolMenuEntry::GetCheckState returns Checked")
+		{
+			CHECK(Entry.GetCheckState(FToolMenuContext()) == ECheckBoxState::Unchecked);
+		}
+	}
+}
+
+TEST_CASE("Developer::ToolMenus::UToolMenuEntry::GetCheckState supports FToolUIAction", "[ToolMenus]")
+{
+	GIVEN("An entry with an FToolUIAction with a GetActionCheckState that returns Checked")
+	{
+		FToolUIAction Action;
+		Action.GetActionCheckState = FToolMenuGetActionCheckState::CreateLambda(
+			[](const FToolMenuContext&) -> ECheckBoxState
+			{
+				return ECheckBoxState::Checked;
+			}
+		);
+
+		const FToolMenuEntry Entry = FToolMenuEntry::InitMenuEntry(NAME_None, FText(), FText(), FSlateIcon(), Action);
+
+		THEN("FToolMenuEntry::GetCheckState returns Checked")
+		{
+			CHECK(Entry.GetCheckState(FToolMenuContext()) == ECheckBoxState::Checked);
+		}
+	}
+
+	GIVEN("An entry with an FToolUIAction with a GetActionCheckState that returns Checked")
+	{
+		FToolUIAction Action;
+		Action.GetActionCheckState = FToolMenuGetActionCheckState::CreateLambda(
+			[](const FToolMenuContext&) -> ECheckBoxState
+			{
+				return ECheckBoxState::Unchecked;
+			}
+		);
+
+		const FToolMenuEntry Entry = FToolMenuEntry::InitMenuEntry(NAME_None, FText(), FText(), FSlateIcon(), Action);
+
+		THEN("FToolMenuEntry::GetCheckState returns Unchecked")
+		{
+			CHECK(Entry.GetCheckState(FToolMenuContext()) == ECheckBoxState::Unchecked);
+		}
+	}
+}
+
+TEST_CASE("Developer::ToolMenus::UToolMenuEntry::GetCheckState supports FToolDynamicUIAction", "[ToolMenus]")
+{
+	GIVEN("An entry with an FToolDynamicUIAction with a GetActionCheckState that returns Checked")
+	{
+		UCheckBoxStateObject* CheckBoxStateObject = NewObject<UCheckBoxStateObject>();
+		CheckBoxStateObject->SetStateToReturn(ECheckBoxState::Checked);
+
+		FToolDynamicUIAction Action;
+		Action.GetActionCheckState.BindUFunction(CheckBoxStateObject, "GetActionCheckState");
+
+		const FToolMenuEntry Entry = FToolMenuEntry::InitMenuEntry(NAME_None, FText(), FText(), FSlateIcon(), Action);
+
+		THEN("FToolMenuEntry::GetCheckState returns Checked")
+		{
+			CHECK(Entry.GetCheckState(FToolMenuContext()) == ECheckBoxState::Checked);
+		}
+	}
+	GIVEN("An entry with an FToolDynamicUIAction with a GetActionCheckState that returns Unchecked")
+	{
+		UCheckBoxStateObject* CheckBoxStateObject = NewObject<UCheckBoxStateObject>();
+		CheckBoxStateObject->SetStateToReturn(ECheckBoxState::Unchecked);
+
+		FToolDynamicUIAction Action;
+		Action.GetActionCheckState.BindUFunction(CheckBoxStateObject, "GetActionCheckState");
+
+		const FToolMenuEntry Entry = FToolMenuEntry::InitMenuEntry(NAME_None, FText(), FText(), FSlateIcon(), Action);
+
+		THEN("FToolMenuEntry::GetCheckState returns Unchecked")
+		{
+			CHECK(Entry.GetCheckState(FToolMenuContext()) == ECheckBoxState::Unchecked);
+		}
+	}
+}
diff --git a/Engine/Source/Programs/ToolMenusTests/Tests/ToolMenuIterationTest.cpp b/Engine/Source/Programs/ToolMenusTests/Tests/ToolMenuIterationTest.cpp
new file mode 100644
--- /dev/null
+++ b/Engine/Source/Programs/ToolMenusTests/Tests/ToolMenuIterationTest.cpp
@@ -0,0 +1,392 @@
+// Copyright Epic Games, Inc. All Rights Reserved.
+
+#include "ToolMenuIteration.h"
+
+#include "TestHarness.h"
+#include "ToolMenus.h"
+
+TEST_CASE("Developer::ToolMenus::FToolMenuIteration does not visit anything for non-existant menus", "[ToolMenus]")
+{
+	UE::ToolMenus::FToolMenuTestInstanceScoped Scope;
+
+	GIVEN("The menu we try to visist does not exist")
+	{
+		const FName MenuName = "MenuNameThatDoesNotExist";
+
+		int32 NumVisists = 0;
+		UE::ToolMenus::VisitMenuEntries(
+			UToolMenus::Get(),
+			MenuName,
+			FToolMenuContext(),
+			UE::ToolMenus::FToolMenuVisitor::CreateLambda(
+				[&NumVisists](const UE::ToolMenus::FToolMenuIterationInfo& Info) -> bool
+				{
+					++NumVisists;
+
+					return true;
+				}
+			)
+		);
+
+		THEN("We did not visit any entries")
+		{
+			CHECK(NumVisists == 0);
+		}
+	}
+}
+
+TEST_CASE("Developer::ToolMenus::FToolMenuIteration can visit menus", "[ToolMenus]")
+{
+	UE::ToolMenus::FToolMenuTestInstanceScoped Scope;
+
+	const FName MenuName = "ToolMenuIterationTest_MyMenu";
+	UToolMenu* const Menu = UToolMenus::Get()->RegisterMenu(MenuName);
+
+	GIVEN("An empty menu")
+	{
+		WHEN("The menu is iterated")
+		{
+			int32 NumVisists = 0;
+			UE::ToolMenus::VisitMenuEntries(
+				UToolMenus::Get(),
+				MenuName,
+				FToolMenuContext(),
+				UE::ToolMenus::FToolMenuVisitor::CreateLambda(
+					[&NumVisists](const UE::ToolMenus::FToolMenuIterationInfo& Info) -> bool
+					{
+						++NumVisists;
+
+						return true;
+					}
+				)
+			);
+
+			THEN("No entries are visisted")
+			{
+				CHECK(NumVisists == 0);
+			}
+		}
+	}
+
+	GIVEN("A menu with one entry")
+	{
+		FToolMenuSection& Section = Menu->FindOrAddSection(NAME_None);
+		Section.AddMenuEntry("MyEntry", FText(), FText(), FSlateIcon(), FUIAction());
+
+		WHEN("The menu is iterated")
+		{
+			int32 NumVisists = 0;
+			UE::ToolMenus::VisitMenuEntries(
+				UToolMenus::Get(),
+				MenuName,
+				FToolMenuContext(),
+				UE::ToolMenus::FToolMenuVisitor::CreateLambda(
+					[&NumVisists](const UE::ToolMenus::FToolMenuIterationInfo& Info) -> bool
+					{
+						++NumVisists;
+
+						return true;
+					}
+				)
+			);
+
+			THEN("One entry is visited")
+			{
+				CHECK(NumVisists == 1);
+			}
+		}
+	}
+	GIVEN("A menu with six entries across two sections")
+	{
+		{
+			FToolMenuSection& Section = Menu->FindOrAddSection(NAME_None);
+			Section.AddMenuEntry("A", FText(), FText(), FSlateIcon(), FUIAction());
+			Section.AddMenuEntry("B", FText(), FText(), FSlateIcon(), FUIAction());
+			Section.AddMenuEntry("C", FText(), FText(), FSlateIcon(), FUIAction());
+		}
+
+		{
+			FToolMenuSection& Section = Menu->FindOrAddSection("Foo");
+			Section.AddMenuEntry("D", FText(), FText(), FSlateIcon(), FUIAction());
+			Section.AddMenuEntry("E", FText(), FText(), FSlateIcon(), FUIAction());
+			Section.AddMenuEntry("F", FText(), FText(), FSlateIcon(), FUIAction());
+		}
+
+		WHEN("The menu is iterated")
+		{
+			TSet<FName> VisistedEntryNames;
+			TSet<FName> VisistedSectionNames;
+
+			int32 NumVisists = 0;
+			UE::ToolMenus::VisitMenuEntries(
+				UToolMenus::Get(),
+				MenuName,
+				FToolMenuContext(),
+				UE::ToolMenus::FToolMenuVisitor::CreateLambda(
+					[&NumVisists, &VisistedEntryNames, &VisistedSectionNames](const UE::ToolMenus::FToolMenuIterationInfo& Info
+					) -> bool
+					{
+						++NumVisists;
+						VisistedEntryNames.Add(Info.Entry.Name);
+						VisistedSectionNames.Add(Info.Section.Name);
+
+						return true;
+					}
+				)
+			);
+
+			THEN("Six entries are found")
+			{
+				CHECK(NumVisists == 6);
+			}
+
+			THEN("The extected entry names were visisted")
+			{
+				CHECK(VisistedEntryNames.Contains("A"));
+				CHECK(VisistedEntryNames.Contains("B"));
+				CHECK(VisistedEntryNames.Contains("C"));
+
+				CHECK(VisistedEntryNames.Contains("D"));
+				CHECK(VisistedEntryNames.Contains("E"));
+				CHECK(VisistedEntryNames.Contains("F"));
+			}
+
+			THEN("The extected section names were visisted")
+			{
+				CHECK(VisistedSectionNames.Contains(NAME_None));
+				CHECK(VisistedSectionNames.Contains("Foo"));
+			}
+		}
+	}
+}
+
+TEST_CASE("Developer::ToolMenus::FToolMenuIteration can visit menu extensions", "[ToolMenus]")
+{
+	UE::ToolMenus::FToolMenuTestInstanceScoped Scope;
+
+	const FName MenuName = "ToolMenuIterationTest_MyMenu";
+	UToolMenu* const Menu = UToolMenus::Get()->RegisterMenu(MenuName);
+
+	GIVEN("A menu with one entry")
+	{
+		FToolMenuSection& Section = Menu->FindOrAddSection(NAME_None);
+		Section.AddMenuEntry("A", FText(), FText(), FSlateIcon(), FUIAction());
+
+		AND_GIVEN("The menu is extended with another entry")
+		{
+			{
+				UToolMenu* const ExtendedMenu = UToolMenus::Get()->ExtendMenu(MenuName);
+				FToolMenuSection& SomeSection = ExtendedMenu->FindOrAddSection("SomeSection");
+				SomeSection.AddMenuEntry("B", FText(), FText(), FSlateIcon(), FUIAction());
+			}
+
+			WHEN("The menu is iterated")
+			{
+				TSet<FName> VisistedEntryNames;
+
+				int32 NumVisists = 0;
+				UE::ToolMenus::VisitMenuEntries(
+					UToolMenus::Get(),
+					MenuName,
+					FToolMenuContext(),
+					UE::ToolMenus::FToolMenuVisitor::CreateLambda(
+						[&NumVisists, &VisistedEntryNames](const UE::ToolMenus::FToolMenuIterationInfo& Info) -> bool
+						{
+							++NumVisists;
+							VisistedEntryNames.Add(Info.Entry.Name);
+
+							return true;
+						}
+					)
+				);
+
+				THEN("Two entries are found")
+				{
+					CHECK(NumVisists == 2);
+				}
+
+				THEN("The entry of the base menu is found")
+				{
+					CHECK(VisistedEntryNames.Contains("A"));
+				}
+
+				THEN("The entry of the extension is found")
+				{
+					CHECK(VisistedEntryNames.Contains("B"));
+				}
+			}
+		}
+	}
+}
+
+TEST_CASE("Developer::ToolMenus::FToolMenuIteration can visit submenu extensions", "[ToolMenus]")
+{
+	// We have to use the global UToolMenus instance because UToolMenu::GetMenuCustomizationHierarchy calls UToolMenus::Get().
+	UToolMenus* const ToolMenus = UToolMenus::Get();
+
+	const FName MenuName = "ToolMenuIterationTest_MyMenu";
+	UToolMenu* const Menu = ToolMenus->RegisterMenu(MenuName);
+
+	GIVEN("A menu with a submenu that has one entry")
+	{
+		const FName SubmenuName = "MySubmenu";
+		{
+			FToolMenuSection& Section = Menu->FindOrAddSection(NAME_None);
+
+			FNewToolMenuDelegate MakeMenu = FNewToolMenuDelegate::CreateLambda(
+				[](UToolMenu* InMenu)
+				{
+					FToolMenuSection& Section = InMenu->FindOrAddSection(NAME_None);
+					Section.AddMenuEntry("one", FText(), FText(), FSlateIcon(), FUIAction());
+				}
+			);
+
+			Section.AddSubMenu(SubmenuName, FText(), FText(), MakeMenu);
+		}
+
+		AND_GIVEN("The submenu is extended with another entry")
+		{
+			{
+				const FName FullSubmenuName = UToolMenus::JoinMenuPaths(MenuName, SubmenuName);
+				UToolMenu* const ExtendedSubmenu = ToolMenus->ExtendMenu(FullSubmenuName);
+				FToolMenuSection& Section = ExtendedSubmenu->FindOrAddSection("Foo");
+				Section.AddMenuEntry("two", FText(), FText(), FSlateIcon(), FUIAction());
+			}
+
+			WHEN("The menu is iterated")
+			{
+				TSet<FName> VisitedEntryNames;
+
+				int32 NumVisists = 0;
+				UE::ToolMenus::VisitMenuEntries(
+					ToolMenus,
+					MenuName,
+					FToolMenuContext(),
+					UE::ToolMenus::FToolMenuVisitor::CreateLambda(
+						[&NumVisists, &VisitedEntryNames](const UE::ToolMenus::FToolMenuIterationInfo& Info) -> bool
+						{
+							++NumVisists;
+							VisitedEntryNames.Add(Info.Entry.Name);
+
+							return true;
+						}
+					)
+				);
+
+				THEN("Two entries are found")
+				{
+					CHECK(NumVisists == 2);
+				}
+
+				THEN("The entry of the base submenu is found")
+				{
+					CHECK(VisitedEntryNames.Contains("one"));
+				}
+
+				THEN("The entry of the extension to the submenu is found")
+				{
+					CHECK(VisitedEntryNames.Contains("two"));
+				}
+			}
+		}
+	}
+
+	ToolMenus->RemoveMenu(MenuName);
+}
+
+TEST_CASE("Developer::ToolMenus::FToolMenuIteration can be stopped", "[ToolMenus]")
+{
+	UE::ToolMenus::FToolMenuTestInstanceScoped Scope;
+
+	const FName MenuName = "ToolMenuIterationTest_MyMenu";
+	UToolMenu* const Menu = UToolMenus::Get()->RegisterMenu(MenuName);
+
+	GIVEN("A menu with two entries")
+	{
+		FToolMenuSection& Section = Menu->FindOrAddSection(NAME_None);
+		Section.AddMenuEntry("one", FText(), FText(), FSlateIcon(), FUIAction());
+		Section.AddMenuEntry("two", FText(), FText(), FSlateIcon(), FUIAction());
+
+		WHEN("The menu is iterated and the visitor returns false to stop the iteration")
+		{
+			int32 NumVisits = 0;
+			UE::ToolMenus::VisitMenuEntries(
+				UToolMenus::Get(),
+				MenuName,
+				FToolMenuContext(),
+				UE::ToolMenus::FToolMenuVisitor::CreateLambda(
+					[&NumVisits](const UE::ToolMenus::FToolMenuIterationInfo& Info) -> bool
+					{
+						++NumVisits;
+
+						return false;
+					}
+				)
+			);
+
+			THEN("One of the two entries is visisted")
+			{
+				CHECK(NumVisits == 1);
+			}
+		}
+	}
+}
+
+TEST_CASE("Developer::ToolMenus::FToolMenuIteration can visit submenus", "[ToolMenus]")
+{
+	// We have to use the global UToolMenus instance because UToolMenu::GetMenuCustomizationHierarchy calls UToolMenus::Get().
+	UToolMenus* const ToolMenus = UToolMenus::Get();
+
+	const FName MenuName = "ToolMenuIterationTest_MyMenuWithSubmenu";
+	UToolMenu* const Menu = ToolMenus->RegisterMenu(MenuName);
+
+	GIVEN("A menu with a submenu that has one entry")
+	{
+		{
+			FToolMenuSection& Section = Menu->FindOrAddSection(NAME_None);
+
+			FNewToolMenuDelegate MakeMenu = FNewToolMenuDelegate::CreateLambda(
+				[](UToolMenu* InMenu)
+				{
+					FToolMenuSection& Section = InMenu->FindOrAddSection(NAME_None);
+					Section.AddMenuEntry("X", FText(), FText(), FSlateIcon(), FUIAction());
+				}
+			);
+
+			Section.AddSubMenu("MySubmenu", FText(), FText(), MakeMenu);
+		}
+
+		WHEN("The menu is iterated")
+		{
+			TSet<FName> VisitedEntryNames;
+
+			int32 NumVisits = 0;
+			UE::ToolMenus::VisitMenuEntries(
+				ToolMenus,
+				MenuName,
+				FToolMenuContext(),
+				UE::ToolMenus::FToolMenuVisitor::CreateLambda(
+					[&NumVisits, &VisitedEntryNames](const UE::ToolMenus::FToolMenuIterationInfo& Info) -> bool
+					{
+						++NumVisits;
+						VisitedEntryNames.Add(Info.Entry.Name);
+
+						return true;
+					}
+				)
+			);
+
+			THEN("One entry is visisted")
+			{
+				CHECK(NumVisits == 1);
+			}
+
+			THEN("The expected entry name was visisted")
+			{
+				CHECK(VisitedEntryNames.Contains("X"));
+			}
+		}
+	}
+
+	ToolMenus->RemoveMenu(MenuName);
+}
diff --git a/Engine/Source/Programs/ToolMenusTests/Tests/ToolMenuTestInstanceScopedTest.cpp b/Engine/Source/Programs/ToolMenusTests/Tests/ToolMenuTestInstanceScopedTest.cpp
new file mode 100644
--- /dev/null
+++ b/Engine/Source/Programs/ToolMenusTests/Tests/ToolMenuTestInstanceScopedTest.cpp
@@ -0,0 +1,67 @@
+// Copyright Epic Games, Inc. All Rights Reserved.
+
+#include "TestHarness.h"
+
+#include "ToolMenus.h"
+
+TEST_CASE("UE::ToolMenus::FToolMenuTestInstanceScoped replaces the singleton while the scope is active", "[ToolMenus]")
+{
+	UToolMenus* const OriginalSingleton = UToolMenus::Get();
+
+	{
+		UE::ToolMenus::FToolMenuTestInstanceScoped Scope1;
+
+		UToolMenus* const FirstNestedSingleton = UToolMenus::Get();
+
+		CHECK(FirstNestedSingleton != OriginalSingleton);
+
+		{
+			UE::ToolMenus::FToolMenuTestInstanceScoped Scope2;
+
+			UToolMenus* const SecondNestedSingleton = UToolMenus::Get();
+
+			CHECK(SecondNestedSingleton != FirstNestedSingleton);
+			CHECK(SecondNestedSingleton != OriginalSingleton);
+		}
+	}
+}
+
+TEST_CASE("UE::ToolMenus::FToolMenuTestInstanceScoped reinstates the original singleton when the scope is destructed", "[ToolMenus]")
+{
+	UToolMenus* const OriginalSingleton = UToolMenus::Get();
+
+	{
+		UE::ToolMenus::FToolMenuTestInstanceScoped Scope1;
+
+		UToolMenus* const FirstNestedSingleton = UToolMenus::Get();
+
+		{
+			UE::ToolMenus::FToolMenuTestInstanceScoped Scope2;
+		}
+
+		CHECK(FirstNestedSingleton == UToolMenus::Get());
+	}
+
+	CHECK(OriginalSingleton == UToolMenus::Get());
+}
+
+const FName FToolMenuTestInstanceScopedMenuName =
+	"FToolMenuTestInstanceScoped_prevents_names_from_leaking_between_tests";
+
+TEST_CASE("UE::ToolMenus::FToolMenuTestInstanceScoped prevents menu names from leaking between tests (test 1 of 2)", "[ToolMenus]")
+{
+	UE::ToolMenus::FToolMenuTestInstanceScoped Scope;
+
+	CHECK(!UToolMenus::Get()->IsMenuRegistered(FToolMenuTestInstanceScopedMenuName));
+
+	UToolMenus::Get()->RegisterMenu(FToolMenuTestInstanceScopedMenuName);
+}
+
+TEST_CASE("UE::ToolMenus::FToolMenuTestInstanceScoped prevents menu names from leaking between tests (test 2 of 2)", "[ToolMenus]")
+{
+	UE::ToolMenus::FToolMenuTestInstanceScoped Scope;
+
+	CHECK(!UToolMenus::Get()->IsMenuRegistered(FToolMenuTestInstanceScopedMenuName));
+
+	UToolMenus::Get()->RegisterMenu(FToolMenuTestInstanceScopedMenuName);
+}
diff --git a/Engine/Source/Programs/ToolMenusTests/Tests/ToolMenus/CheckBoxStateObject.h b/Engine/Source/Programs/ToolMenusTests/Tests/ToolMenus/CheckBoxStateObject.h
deleted file mode 100644
diff --git a/Engine/Source/Programs/ToolMenusTests/Tests/ToolMenus/ToolMenuEntryTest.cpp b/Engine/Source/Programs/ToolMenusTests/Tests/ToolMenus/ToolMenuEntryTest.cpp
deleted file mode 100644
diff --git a/Engine/Source/Programs/ToolMenusTests/Tests/ToolMenus/ToolMenuIterationTest.cpp b/Engine/Source/Programs/ToolMenusTests/Tests/ToolMenus/ToolMenuIterationTest.cpp
deleted file mode 100644
diff --git a/Engine/Source/Programs/ToolMenusTests/Tests/ToolMenus/ToolMenuTestInstanceScopedTest.cpp b/Engine/Source/Programs/ToolMenusTests/Tests/ToolMenus/ToolMenuTestInstanceScopedTest.cpp
deleted file mode 100644
diff --git a/Engine/Source/Programs/ToolMenusTests/Tests/ToolMenus/ToolMenusTest.cpp b/Engine/Source/Programs/ToolMenusTests/Tests/ToolMenus/ToolMenusTest.cpp
deleted file mode 100644
diff --git a/Engine/Source/Programs/ToolMenusTests/Tests/ToolMenusTest.cpp b/Engine/Source/Programs/ToolMenusTests/Tests/ToolMenusTest.cpp
new file mode 100644
--- /dev/null
+++ b/Engine/Source/Programs/ToolMenusTests/Tests/ToolMenusTest.cpp
@@ -0,0 +1,395 @@
+// Copyright Epic Games, Inc. All Rights Reserved.
+
+#include "ToolMenus.h"
+
+#include "TestHarness.h"
+
+TEST_CASE("Developer::ToolMenus::Can create menu", "[ToolMenus]")
+{
+	UE::ToolMenus::FToolMenuTestInstanceScoped Scope;
+
+	UToolMenu* ToolMenu = UToolMenus::Get()->RegisterMenu("ToolMenusTest_MyMenu");
+
+	CHECK(ToolMenu);
+}
+
+TEST_CASE("Developer::ToolMenus::Non-registered menu is not registered", "[ToolMenus]")
+{
+	UE::ToolMenus::FToolMenuTestInstanceScoped Scope;
+
+	CHECK_FALSE(UToolMenus::Get()->IsMenuRegistered("ToolMenusTest_MyMenu"));
+}
+
+TEST_CASE("Developer::ToolMenus::Removed menu is not registered", "[ToolMenus]")
+{
+	UE::ToolMenus::FToolMenuTestInstanceScoped Scope;
+
+	const FName MenuName("ToolMenusTest_MyMenu");
+
+	UToolMenus::Get()->RegisterMenu(MenuName);
+
+	CHECK(UToolMenus::Get()->IsMenuRegistered(MenuName));
+
+	UToolMenus::Get()->RemoveMenu(MenuName);
+
+	CHECK_FALSE(UToolMenus::Get()->IsMenuRegistered(MenuName));
+}
+
+TEST_CASE("Developer::ToolMenus::GenerateWidget calls dynamic section lambdas", "[ToolMenus]")
+{
+	UE::ToolMenus::FToolMenuTestInstanceScoped Scope;
+
+	const FName MenuName("ToolMenusTest_MyMenu");
+
+	UToolMenu* ToolMenu = UToolMenus::Get()->RegisterMenu(MenuName);
+
+	REQUIRE(ToolMenu);
+
+	bool bWasLambdaCalled = false;
+	ToolMenu->AddDynamicSection("MySection", FNewToolMenuDelegateLegacy::CreateLambda([&bWasLambdaCalled](FMenuBuilder&, UToolMenu*) {
+		bWasLambdaCalled = true;
+	}));
+
+	UToolMenus::Get()->GenerateWidget(MenuName, FToolMenuContext());
+
+	CHECK(bWasLambdaCalled);
+}
+
+TEST_CASE("Developer::ToolMenus::Can add and remove sections to menu", "[ToolMenus]")
+{
+	UE::ToolMenus::FToolMenuTestInstanceScoped Scope;
+
+	UToolMenu* ToolMenu = UToolMenus::Get()->RegisterMenu("ToolMenusTest_MyMenu");
+
+	REQUIRE(ToolMenu);
+
+	// Add a section to the menu
+	const FName SectionName("MySection");
+	
+	ToolMenu->AddSection(SectionName);
+	CHECK(ToolMenu->Sections.Num() == 1);
+
+	// Check if the section we added is correct
+	FToolMenuSection* Section = ToolMenu->FindSection(SectionName);
+	REQUIRE(Section);
+	CHECK(Section->Name == SectionName);
+
+	// Check if we are able to remove the section
+	ToolMenu->RemoveSection(SectionName);
+	CHECK(ToolMenu->Sections.Num() == 0);
+}
+
+TEST_CASE("Developer::ToolMenus::Can add entries to menu section", "[ToolMenus]")
+{
+	UE::ToolMenus::FToolMenuTestInstanceScoped Scope;
+
+	UToolMenu* ToolMenu = UToolMenus::Get()->RegisterMenu("ToolMenusTest_MyMenu");
+
+	REQUIRE(ToolMenu);
+
+	FToolMenuSection& Section = ToolMenu->AddSection("MySection");
+
+	// Add the first entry to the menu
+	const FName Entry1Name("MyEntry1");
+	Section.AddMenuEntry(Entry1Name, FText::GetEmpty(), FText::GetEmpty(), FSlateIcon(), FToolUIActionChoice());
+	CHECK(Section.Blocks.Num() == 1);
+
+	// Check if the first entry is correct
+	FToolMenuEntry* Entry1 = Section.FindEntry(Entry1Name);
+	REQUIRE(Entry1);
+	CHECK(Entry1->Name == Entry1Name);
+
+	// Add the second entry to the menu
+	const FName Entry2Name("MyEntry2");
+	Section.AddMenuEntry(Entry2Name, FText::GetEmpty(), FText::GetEmpty(), FSlateIcon(), FToolUIActionChoice());
+	CHECK(Section.Blocks.Num() == 2);
+
+	// Check if the second entry is correct
+	FToolMenuEntry* Entry2 = Section.FindEntry(Entry2Name);
+	REQUIRE(Entry2);
+	CHECK(Entry2->Name == Entry2Name);
+}
+
+TEST_CASE("Developer::ToolMenus::Can add runtime menu customization", "[ToolMenus]")
+{
+	UE::ToolMenus::FToolMenuTestInstanceScoped Scope;
+
+	const FName MenuName("MyMenu");
+
+	UToolMenus::Get()->RegisterMenu(MenuName);
+
+	// Add the customization to the menu
+	UToolMenus::Get()->AddRuntimeMenuCustomization(MenuName);
+
+	// Try and find the customization
+	FCustomizedToolMenu* CustomizedToolMenu = UToolMenus::Get()->FindRuntimeMenuCustomization(MenuName);
+	CHECK(CustomizedToolMenu);
+}
+
+
+TEST_CASE("Developer::ToolMenus::Can customize menu using runtime menu customization", "[ToolMenus]")
+{
+	UE::ToolMenus::FToolMenuTestInstanceScoped Scope;
+
+	const FName MenuName("MyMenu");
+
+	UToolMenu* ToolMenu = UToolMenus::Get()->RegisterMenu(MenuName);
+
+	REQUIRE(ToolMenu);
+
+	// Add test entries 1, 2 and 3
+	const FName SectionName("MySection");
+	const FName Entry1Name("MyEntry1");
+	const FName Entry2Name("MyEntry2");
+	const FName Entry3Name("MyEntry3");
+
+	FToolMenuSection& Section = ToolMenu->AddSection(SectionName);
+	Section.AddMenuEntry(Entry1Name, FText::GetEmpty(), FText::GetEmpty(), FSlateIcon(), FToolUIActionChoice());
+	Section.AddMenuEntry(Entry2Name, FText::GetEmpty(), FText::GetEmpty(), FSlateIcon(), FToolUIActionChoice());
+	Section.AddMenuEntry(Entry3Name, FText::GetEmpty(), FText::GetEmpty(), FSlateIcon(), FToolUIActionChoice());
+
+	// Add the customization to the menu
+	UToolMenus::Get()->AddRuntimeMenuCustomization(MenuName);
+	FCustomizedToolMenu* CustomizedToolMenu = UToolMenus::Get()->FindRuntimeMenuCustomization(MenuName);
+	REQUIRE(CustomizedToolMenu);
+
+	// Test the customization by denying entry 2
+	const FName OwnerName("ToolMenusTest");
+
+	CustomizedToolMenu->MenuPermissions.AddDenyListItem(OwnerName, Entry2Name);
+
+	UToolMenu* GeneratedMenu = UToolMenus::Get()->GenerateMenu(MenuName, FToolMenuContext());
+	REQUIRE(GeneratedMenu);
+
+	FToolMenuSection* GeneratedMenuSection = GeneratedMenu->FindSection(SectionName);
+	REQUIRE(GeneratedMenuSection);
+
+	// We should only have 2 entries: Entry1 and Entry3 since Entry2 is denied
+	CHECK(GeneratedMenuSection->Blocks.Num() == 2);
+	CHECK(GeneratedMenuSection->FindEntry(Entry1Name));
+	CHECK_FALSE(GeneratedMenuSection->FindEntry(Entry2Name));
+	CHECK(GeneratedMenuSection->FindEntry(Entry3Name));
+}
+
+TEST_CASE("Developer::ToolMenus::Can add runtime menu profile", "[ToolMenus]")
+{
+	UE::ToolMenus::FToolMenuTestInstanceScoped Scope;
+
+	const FName MenuName("MyMenu");
+
+	UToolMenus::Get()->RegisterMenu(MenuName);
+
+	// Add the customization to the menu
+	const FName ProfileName("MyProfile1");
+	UToolMenus::Get()->AddRuntimeMenuProfile(MenuName, ProfileName);
+
+	// Try and find the profile
+	FToolMenuProfile* MenuProfile = UToolMenus::Get()->FindRuntimeMenuProfile(MenuName, ProfileName);
+	CHECK(MenuProfile);
+}
+
+TEST_CASE("Developer::ToolMenus::Can customize menu using runtime menu profile", "[ToolMenus]")
+{
+	UE::ToolMenus::FToolMenuTestInstanceScoped Scope;
+
+	const FName MenuName("MyMenu");
+
+	UToolMenu* ToolMenu = UToolMenus::Get()->RegisterMenu(MenuName);
+
+	REQUIRE(ToolMenu);
+
+	// Add test entries 1, 2 and 3 to the menu
+	const FName SectionName("MySection");
+	const FName Entry1Name("MyEntry1");
+	const FName Entry2Name("MyEntry2");
+	const FName Entry3Name("MyEntry3");
+
+	FToolMenuSection& Section = ToolMenu->AddSection(SectionName);
+	Section.AddMenuEntry(Entry1Name, FText::GetEmpty(), FText::GetEmpty(), FSlateIcon(), FToolUIActionChoice());
+	Section.AddMenuEntry(Entry2Name, FText::GetEmpty(), FText::GetEmpty(), FSlateIcon(), FToolUIActionChoice());
+	Section.AddMenuEntry(Entry3Name, FText::GetEmpty(), FText::GetEmpty(), FSlateIcon(), FToolUIActionChoice());
+
+	// Add a profile to the menu
+	const FName Profile1Name("MyProfile1");
+
+	UToolMenus::Get()->AddRuntimeMenuProfile(MenuName, Profile1Name);
+
+	// Check if the profile was added
+	FToolMenuProfile* MenuProfile = UToolMenus::Get()->FindRuntimeMenuProfile(MenuName, Profile1Name);
+	REQUIRE(MenuProfile);
+
+	// Deny Entry 2 in the profile
+	const FName OwnerName("ToolMenusTest");
+	
+	MenuProfile->MenuPermissions.AddDenyListItem(OwnerName, Entry2Name);
+
+	SECTION("Generating menu without profile should show all entries")
+	{
+		UToolMenu* GeneratedMenuWithoutProfile = UToolMenus::Get()->GenerateMenu(MenuName, FToolMenuContext());
+		REQUIRE(GeneratedMenuWithoutProfile);
+
+		FToolMenuSection* GeneratedMenuSection = GeneratedMenuWithoutProfile->FindSection(SectionName);
+		REQUIRE(GeneratedMenuSection);
+
+		CHECK(GeneratedMenuSection->Blocks.Num() == 3);
+		CHECK(GeneratedMenuSection->FindEntry(Entry1Name));
+		CHECK(GeneratedMenuSection->FindEntry(Entry2Name));
+		CHECK(GeneratedMenuSection->FindEntry(Entry3Name));
+	}
+
+	SECTION("Generating menu with profile should hide entry 2")
+	{
+		FToolMenuContext MenuContext;
+
+		UToolMenuProfileContext* ProfileContext = NewObject<UToolMenuProfileContext>();
+		ProfileContext->ActiveProfiles.Add(Profile1Name);
+		MenuContext.AddObject(ProfileContext);
+
+		UToolMenu* GeneratedMenuWithProfile = UToolMenus::Get()->GenerateMenu(MenuName, MenuContext);
+		REQUIRE(GeneratedMenuWithProfile);
+
+		FToolMenuSection* GeneratedMenuSection = GeneratedMenuWithProfile->FindSection(SectionName);
+		REQUIRE(GeneratedMenuSection);
+
+		CHECK(GeneratedMenuSection->Blocks.Num() == 2);
+		CHECK(GeneratedMenuSection->FindEntry(Entry1Name));
+		CHECK_FALSE(GeneratedMenuSection->FindEntry(Entry2Name));
+		CHECK(GeneratedMenuSection->FindEntry(Entry3Name));
+	}
+}
+
+TEST_CASE("Developer::ToolMenus::Can combine customization and multiple profiles", "[ToolMenus]")
+{
+	UE::ToolMenus::FToolMenuTestInstanceScoped Scope;
+
+	const FName MenuName("MyMenu");
+
+	UToolMenu* ToolMenu = UToolMenus::Get()->RegisterMenu(MenuName);
+
+	REQUIRE(ToolMenu);
+
+	// Add the section and entries
+	const FName SectionName("MySection");
+	const FName Entry1Name("MyEntry1");
+	const FName Entry2Name("MyEntry2");
+	const FName Entry3Name("MyEntry3");
+
+	FToolMenuSection& Section = ToolMenu->AddSection(SectionName);
+	Section.AddMenuEntry(Entry1Name, FText::GetEmpty(), FText::GetEmpty(), FSlateIcon(), FToolUIActionChoice());
+	Section.AddMenuEntry(Entry2Name, FText::GetEmpty(), FText::GetEmpty(), FSlateIcon(), FToolUIActionChoice());
+	Section.AddMenuEntry(Entry3Name, FText::GetEmpty(), FText::GetEmpty(), FSlateIcon(), FToolUIActionChoice());
+
+	// The menu customization denies entry 1
+	const FName OwnerName("ToolMenusTest");
+
+	UToolMenus::Get()->AddRuntimeMenuCustomization(MenuName);
+	FCustomizedToolMenu* CustomizedToolMenu = UToolMenus::Get()->FindRuntimeMenuCustomization(MenuName);
+	REQUIRE(CustomizedToolMenu);
+	CustomizedToolMenu->MenuPermissions.AddDenyListItem(OwnerName, Entry2Name);
+
+	// Add Profile 1 which denies entry 1 and 2
+	const FName Profile1Name("MyProfile1");
+
+	UToolMenus::Get()->AddRuntimeMenuProfile(MenuName, Profile1Name);
+	FToolMenuProfile* MenuProfile = UToolMenus::Get()->FindRuntimeMenuProfile(MenuName, Profile1Name);
+	REQUIRE(MenuProfile);
+	MenuProfile->MenuPermissions.AddDenyListItem(OwnerName, Entry1Name);
+	MenuProfile->MenuPermissions.AddDenyListItem(OwnerName, Entry2Name);
+
+	// Add Profile 2 which ONLY allows entry 1 (i.e denies ALL other entries)
+	const FName Profile2Name("MyProfile2");
+
+	UToolMenus::Get()->AddRuntimeMenuProfile(MenuName, Profile2Name);
+	FToolMenuProfile* MenuProfile2 = UToolMenus::Get()->FindRuntimeMenuProfile(MenuName, Profile2Name);
+	REQUIRE(MenuProfile2);
+	MenuProfile2->MenuPermissions.AddAllowListItem(OwnerName, Entry1Name);
+
+
+	SECTION("Generating menu with customization + profile 1 should only show entry 3")
+	{
+		FToolMenuContext MenuContext;
+		
+		UToolMenuProfileContext* ProfileContext = NewObject<UToolMenuProfileContext>();
+		ProfileContext->ActiveProfiles.Add(Profile1Name);
+		MenuContext.AddObject(ProfileContext);
+
+		UToolMenu* GeneratedMenuWithProfile = UToolMenus::Get()->GenerateMenu(MenuName, MenuContext);
+		REQUIRE(GeneratedMenuWithProfile);
+
+		FToolMenuSection* GeneratedMenuSection = GeneratedMenuWithProfile->FindSection(SectionName);
+		REQUIRE(GeneratedMenuSection);
+
+		CHECK(GeneratedMenuSection->Blocks.Num() == 1);
+		CHECK_FALSE(GeneratedMenuSection->FindEntry(Entry1Name));
+		CHECK_FALSE(GeneratedMenuSection->FindEntry(Entry2Name));
+		CHECK(GeneratedMenuSection->FindEntry(Entry3Name));
+	}
+	
+	SECTION("Generating menu with customization + profile 2 should only show entry 1")
+	{
+		FToolMenuContext MenuContext;
+		
+		UToolMenuProfileContext* ProfileContext = NewObject<UToolMenuProfileContext>();
+		ProfileContext->ActiveProfiles.Add(Profile2Name);
+		MenuContext.AddObject(ProfileContext);
+
+		UToolMenu* GeneratedMenuWithProfile = UToolMenus::Get()->GenerateMenu(MenuName, MenuContext);
+		REQUIRE(GeneratedMenuWithProfile);
+
+		FToolMenuSection* GeneratedMenuSection = GeneratedMenuWithProfile->FindSection(SectionName);
+		REQUIRE(GeneratedMenuSection);
+
+		CHECK(GeneratedMenuSection->Blocks.Num() == 1);
+		
+		CHECK(GeneratedMenuSection->FindEntry(Entry1Name));
+		CHECK_FALSE(GeneratedMenuSection->FindEntry(Entry2Name));
+		CHECK_FALSE(GeneratedMenuSection->FindEntry(Entry3Name));
+	}
+	
+	SECTION("Generating menu with customization + profile 1 + profile 2 should only no entries")
+	{
+		FToolMenuContext MenuContext;
+
+		UToolMenuProfileContext* ProfileContext = NewObject<UToolMenuProfileContext>();
+		ProfileContext->ActiveProfiles.Add(Profile1Name);
+		ProfileContext->ActiveProfiles.Add(Profile2Name);
+		MenuContext.AddObject(ProfileContext);
+
+		UToolMenu* GeneratedMenuWithProfile = UToolMenus::Get()->GenerateMenu(MenuName, MenuContext);
+		REQUIRE(GeneratedMenuWithProfile);
+
+		FToolMenuSection* GeneratedMenuSection = GeneratedMenuWithProfile->FindSection(SectionName);
+		REQUIRE(GeneratedMenuSection);
+
+		CHECK(GeneratedMenuSection->Blocks.Num() == 0);
+		CHECK_FALSE(GeneratedMenuSection->FindEntry(Entry1Name));
+		CHECK_FALSE(GeneratedMenuSection->FindEntry(Entry2Name));
+		CHECK_FALSE(GeneratedMenuSection->FindEntry(Entry3Name));
+	}
+}
+
+// Repro of UE-201151.
+TEST_CASE("Developer::ToolMenus::GenerateWidget can handle simultaneous AddReferencedObjects calls in legacy dynamic sections", "[ToolMenus]")
+{
+	UE::ToolMenus::FToolMenuTestInstanceScoped Scope;
+
+	UToolMenu* ToolMenu = UToolMenus::Get()->RegisterMenu("ToolMenusTest_MyMenu");
+
+	REQUIRE(ToolMenu);
+
+	// This simulates the crash of UE-201151 that occurred in UToolMenus::GenerateWidget(UToolMenu*) after a complex delegate
+	// triggered a call to UToolMenus::AddReferencedObjects while the delegate was still executing.
+	ToolMenu->AddDynamicSection(
+		"MyDynamicLegacySection",
+		FNewToolMenuDelegateLegacy::CreateLambda(
+			[](FMenuBuilder&, UToolMenu*)
+			{
+				TArray<UObject*> Array;
+				FReferenceFinder Finder(Array);
+				UToolMenus::AddReferencedObjects(UToolMenus::Get(), Finder);
+			}
+		)
+	);
+
+	UToolMenus::Get()->GenerateWidget("ToolMenusTest_MyMenu", FToolMenuContext());
+}
